Implements r.DumpGPU command

diff --git a/Engine/Extras/GPUDumpViewer/GPUDumpViewer.html b/Engine/Extras/GPUDumpViewer/GPUDumpViewer.html
new file mode 100644
--- /dev/null
+++ b/Engine/Extras/GPUDumpViewer/GPUDumpViewer.html
@@ -0,0 +1,1307 @@
+<html>
+<title>GPU Dump Viewer</title>
+
+<script type="text/javascript">
+
+//var current_dir="file:///E:/frosty/QAGame/Saved/GPUDumps/WindowsEditor/2021.10.04-10.39.34/";
+var current_dir="./";
+
+var g_infos = {};
+var g_passes = [];
+
+
+// -------------------------------------------------------------------- FILE LOADING
+
+function load_text_file(relative_path)
+{
+	try
+	{
+		var request = new XMLHttpRequest();
+		request.open('GET', current_dir + relative_path, false);
+		request.send(null);
+		if (request.status === 0 || request.status === 200)
+		{
+			return request.responseText;
+		}
+	}
+	catch (error)
+	{
+		return null;
+	}
+	return null;
+}
+
+function load_binary_file(relative_path, callback)
+{
+	try
+	{
+		var request = new XMLHttpRequest();
+		request.open('GET', current_dir + relative_path, true);
+		request.responseType = "arraybuffer";
+
+		request.onload = function(event)
+		{
+			var array_buffer = request.response;
+			if ((request.status === 0 || request.status === 200) && array_buffer)
+			{
+				callback(array_buffer);
+			}
+			else
+			{
+				callback(null);
+			}
+		};
+
+		request.onerror = function() {
+			callback(null);
+		};
+
+		request.send(null);
+	}
+	catch (error)
+	{
+		callback(null);
+	}
+}
+
+function load_json(relative_path)
+{
+	var txt = load_text_file(relative_path);
+
+	if (txt === null)
+	{
+		return null;
+	}
+
+	return JSON.parse(load_text_file(relative_path));
+}
+
+function load_json_dict_sequence(relative_path)
+{
+	var text_file = load_text_file(relative_path);
+	var dicts = text_file.split("}{");
+	var dict_sequence = [];
+	dicts.forEach(function(value, index, array) {
+		if (!value.startsWith('{'))
+		{
+			value = '{' + value;
+		}
+		if (!value.endsWith('}'))
+		{
+			value = value + '}';
+		}
+		dict_sequence.push(JSON.parse(value));
+	});
+	return dict_sequence;
+}
+
+function load_resource_desc(resource_name)
+{
+	return load_json(`Resources/${resource_name}.json`);
+}
+
+
+// -------------------------------------------------------------------- DISPLAY PASS
+
+var g_pass_id = null;
+
+function display_pass_hierarchy()
+{
+	var search_pass = document.getElementById('pass_search_input').value;
+	var search_resource = document.getElementById('resource_search_input').value;
+	var html = '';
+	var parent_event_scopes = [];
+
+	g_passes.forEach(function(pass_data, pass_id) {
+		var pass_event_scopes = pass_data['ParentEventScopes'];
+
+		var show_pass = true;
+		if (search_pass != '')
+		{
+			show_pass = pass_data['EventName'].toLowerCase().includes(search_pass.toLowerCase());
+
+			for (var i = 0; i < pass_event_scopes.length; i++)
+			{
+				show_pass = show_pass || pass_event_scopes[i].toLowerCase().includes(search_pass.toLowerCase());
+			}
+		}
+
+		var show_resource = true;
+		if (search_resource != '')
+		{
+			show_resource = false;
+			pass_data['OutputResources'].forEach(function(output_name) {
+				show_resource = show_resource || output_name.toLowerCase().includes(search_resource.toLowerCase());
+			});
+		}
+
+		var has_input_or_outputs = pass_data['InputResources'].length > 0 || pass_data['OutputResources'].length > 0;
+
+		if (show_pass && show_resource && has_input_or_outputs)
+		{
+			var shared_scope = 0;
+			for (var i = 0; i < Math.min(parent_event_scopes.length, pass_event_scopes.length); i++)
+			{
+				if (parent_event_scopes[i] == pass_event_scopes[pass_event_scopes.length - 1 - i])
+				{
+					shared_scope++;
+				}
+				else
+				{
+					break;
+				}
+			}
+
+			parent_event_scopes = parent_event_scopes.slice(0, shared_scope);
+
+			for (var i = shared_scope; i < pass_event_scopes.length; i++)
+			{
+				var scope = pass_event_scopes[pass_event_scopes.length - 1 - i];
+				html += `<div style="padding-left: ${10 + 16 * i}px;" class="disabled">${scope}</div>`;
+				parent_event_scopes.push(scope);
+			}
+
+			var classes = '';
+			if (pass_id === g_pass_id)
+			{
+				classes = 'selected';
+			}
+
+			html += `<div style="padding-left: ${10 + 16 * pass_event_scopes.length}px;" onmousedown="display_pass(${pass_id});" class="${classes}">${pass_data['EventName']}</div>`;
+		}
+	});
+	document.getElementById('pass_hierarchy').innerHTML = html;
+}
+
+function display_input_resource(pass_id, resource_name)
+{
+	var previous_producer_pass = -1;
+	for (var i = 0; i < pass_id; i++)
+	{
+		var cur_outputs = g_passes[i]['OutputResources'];
+
+		cur_outputs.forEach(function(value) {
+			if (value == resource_name)
+			{
+				previous_producer_pass = i;
+			}
+		});
+	}
+
+	return display_output_resource(previous_producer_pass, resource_name);
+}
+
+function display_output_resource(pass_id, resource_name)
+{
+	var resource_version_name = resource_name;
+	if (pass_id >= 0)
+	{
+		var pass_data = g_passes[pass_id];
+		resource_version_name = `${resource_name}.v${pass_data['Pointer']}`;
+	}
+
+	var resource_desc = load_resource_desc(resource_name);
+	if (!resource_desc)
+	{
+		return;
+	}
+
+	if (resource_desc['Type'] == 'Texture2D')
+	{
+		return display_texture(resource_name, resource_version_name);
+	}
+	else if (resource_desc['Type'] == 'VertexBuffer' || resource_desc['Type'] == 'StructuredBuffer')
+	{
+		return display_buffer(resource_name, resource_version_name);
+	}
+}
+
+function display_pass(pass_id)
+{
+	if (pass_id == g_pass_id)
+	{
+		return;
+	}
+
+	var pass_data = g_passes[pass_id];
+	document.getElementById('main_right_pannel').innerHTML = `<div class="pass_title main_div">${pass_data['EventName']}</div>`;
+
+	// Display inputs
+	{
+		var html = `<div class="main_div">
+			<div class="selection_list_title">Input resources: ${pass_data['EventName']}</div>
+			<div class="selection_list">`;
+		pass_data['InputResources'].forEach(function(resource_name) {
+			var resource_desc = load_resource_desc(resource_name);
+			if (resource_desc)
+			{
+				html += `<div onmousedown="display_input_resource(${pass_id}, '${resource_name}');">${resource_desc['Name']}</div>`;
+			}
+			else
+			{
+				html += `<div>${resource_name}</div>`;
+			}
+		});
+		html += '</div></div>';
+		document.getElementById('main_right_pannel').innerHTML += html;
+	}
+
+	// Display outputs
+	{
+		var html = `<div class="main_div">
+			<div class="selection_list_title">Output resources: ${pass_data['EventName']}</div>
+			<div class="selection_list">`;
+		pass_data['OutputResources'].forEach(function(resource_name) {
+			var resource_desc = load_resource_desc(resource_name);
+			html += `<div onmousedown="display_output_resource(${pass_id}, '${resource_name}');">${resource_desc['Name']}</div>`;
+		});
+		html += '</div></div>';
+		document.getElementById('main_right_pannel').innerHTML += html;
+	}
+
+	document.getElementById('main_right_pannel').innerHTML += `<div id="display_resource_pannel"></div>`;
+
+	g_pass_id = pass_id;
+
+	// Load first resource
+	{
+		var resource_name = pass_data['OutputResources'][0];
+
+		if (resource_name)
+		{
+			display_output_resource(pass_id, resource_name);
+		}
+	}
+
+	display_pass_hierarchy();
+}
+
+
+// -------------------------------------------------------------------- DISPLAY GENERAL RESOURCE INFO
+
+function get_resource_producer_html(resource_name)
+{
+	var resource_desc = load_resource_desc(resource_name);
+
+	var html = `
+		<div class="main_div">
+			<div class="selection_list_title">Passes producing ${resource_desc['Name']}</div>
+			<div class="selection_list">`;
+
+	for (var pass_id = 0; pass_id < g_passes.length; pass_id++)
+	{
+		var producer = false;
+		g_passes[pass_id]['OutputResources'].forEach(function(value) {
+			if (value == resource_name)
+			{
+				producer = true;
+			}
+		});
+
+		if (producer)
+		{
+			html += `<div onmousedown="display_pass(${pass_id}); display_output_resource(${pass_id}, '${resource_name}');">${g_passes[pass_id]['EventName']}</div>`;
+		}
+	}
+
+	html += `
+			</div>
+		</div>`;
+
+	html += `
+		<div class="main_div">
+			<div class="selection_list_title">Passes reading ${resource_desc['Name']}</div>
+			<div class="selection_list">`;
+
+	for (var pass_id = 0; pass_id < g_passes.length; pass_id++)
+	{
+		var reader = false;
+		g_passes[pass_id]['InputResources'].forEach(function(value) {
+			if (value == resource_name)
+			{
+				reader = true;
+			}
+		});
+
+		if (reader)
+		{
+			html += `<div onmousedown="display_pass(${pass_id}); display_input_resource(${pass_id}, '${resource_name}');">${g_passes[pass_id]['EventName']}</div>`;
+		}
+	}
+
+	html += `
+			</div>
+		</div>`;
+
+	return html;
+}
+
+
+// -------------------------------------------------------------------- DISPLAY TEXTURE
+
+var g_texture_name = null;
+var g_texture_version_name = null;
+var g_gl = null;
+var g_gl_ext = null;
+var g_gl_texture = null;
+
+var g_shader_code_dict = {};
+
+function get_default_texture_shader_code(gl, resource_desc, resource_version_name)
+{
+	if (resource_version_name in g_shader_code_dict)
+	{
+		return g_shader_code_dict[resource_version_name];
+	}
+
+	var gl_format = translate_pixel_format(gl, resource_desc['Format']);
+
+	if (!gl_format)
+	{
+		return `// Sorry but ${resource_desc['Format']} is not supported for display :(`;
+	}
+
+	var shader_program = `// WebGL visualization shader for a ${gl_format['ue_pixel_format']}
+uniform ${gl_format['sampler_type']} tex_sampler;
+
+in vec2 tex_coordinates;
+out vec4 color;
+
+void main(void)
+{
+	ivec2 texel_coord = ivec2(tex_coordinates * vec2(textureSize(tex_sampler, 0)));
+	${gl_format['sampler_return_type']} raw = texelFetch(tex_sampler, texel_coord , 0);
+	color = ${gl_format['swizzling']};
+}`;
+	return shader_program;
+}
+
+function onresize_texture_shader_user_code_change()
+{
+	refresh_texture_visualization(g_texture_name, g_texture_version_name, null);
+}
+
+function display_texture(resource_name, resource_version_name)
+{
+	var resource_desc = load_resource_desc(resource_name);
+
+	if (resource_version_name == g_texture_version_name)
+	{
+		return;
+	}
+
+	g_texture_name = resource_name;
+	g_texture_version_name = resource_version_name;
+	g_gl = null;
+	g_gl_ext = null;
+	g_gl_texture = null;
+
+	var canvas_w = resource_desc['ExtentX'];
+	var canvas_h = resource_desc['ExtentY'];
+
+	if (canvas_w < 512)
+	{
+		canvas_h = 512 * canvas_h / canvas_w;
+		canvas_w = 512;
+	}
+
+	var resource_info_htmls = '';
+	{
+		resource_info_htmls += `
+			<table width="100%" class="resource_desc">`;
+		
+		for (var key in resource_desc){
+		    if (resource_desc.hasOwnProperty(key)){
+				resource_info_htmls += `
+					<tr>
+						<td>${key}</td>
+						<td>${resource_desc[key]}</td>
+					</tr>`;
+		    }
+		}
+
+		resource_info_htmls += `
+			</table>`;
+	}
+
+	var html = '';
+	html += `
+		<div class="main_div">
+			<div class="selection_list_title">Texture visualization: ${resource_desc['Name']}</div>
+			<div><canvas id="texture_visualization_canvas" width="${canvas_w}px" height="${canvas_h}px"></canvas></div>
+			<table width="100%">
+				<tr>
+					<td>
+						<div class="selection_list_title">Texture descriptor</div>
+						${resource_info_htmls}
+					</td>
+					<td style="width: 50%;">
+						<textarea id="texture_visualization_code_input" oninput="onresize_texture_shader_user_code_change();" onchange="onresize_texture_shader_user_code_change();"></textarea>
+						<textarea id="texture_visualization_code_log" readonly></textarea>
+					</td>
+				</tr>
+			</table>
+		</div>`;
+
+	document.getElementById('display_resource_pannel').innerHTML = html + get_resource_producer_html(resource_name);
+
+	var canvas = document.getElementById('texture_visualization_canvas');
+	g_gl = canvas.getContext('webgl2');
+
+	{
+		var available_extensions = g_gl.getSupportedExtensions();
+		var required_extensions = ['EXT_texture_norm16'];
+
+		g_gl_ext = {};
+		required_extensions.forEach(function(ext_name)
+		{
+			g_gl_ext[ext_name] = g_gl.getExtension(ext_name);
+		});
+	}
+
+	document.getElementById('texture_visualization_code_input').value = get_default_texture_shader_code(g_gl, resource_desc, resource_version_name);
+
+	load_binary_file(`Resources/${resource_version_name}.bin`, function(raw_pixel_data)
+	{
+		if (raw_pixel_data)
+		{
+			refresh_texture_visualization(resource_name, resource_version_name, raw_pixel_data);
+		}
+		else
+		{
+
+		}
+	});
+}
+
+function translate_pixel_format(gl, ue_pixel_format)
+{
+
+	// Missing formats:
+	//PF_G16R16F_FILTER
+	//PF_A2B10G10R10
+	//PF_D24
+	//PF_R16F_FILTER
+	//PF_V8U8
+	//PF_R32_UINT
+	//PF_R32_SINT
+	//PF_PVRTC2
+	//PF_PVRTC4
+	//PF_R16_UINT
+	//PF_R16_SINT
+	//PF_R16G16B16A16_UINT
+	//PF_R16G16B16A16_SINT
+	//PF_R5G6B5_UNORM
+	//PF_R32G32B32A32_UINT
+	//PF_R16G16_UINT
+	//PF_BC6H
+	//PF_R8G8B8A8_SNORM
+	//PF_R16G16B16A16_UNORM
+	//PF_R16G16B16A16_SNORN
+	//PF_R32G32_UINT
+
+	var internal_format = null;
+	var src_format = null;
+	var src_type = null;
+	var swizzling = 'raw';
+	var sampler_type = 'sampler2D';
+	var sampler_return_type = 'vec4';
+
+	// 32bit floats
+	if (ue_pixel_format === 'PF_A32B32G32R32F')
+	{
+		internal_format = gl.RGBA32F;
+		src_format = gl.RGBA;
+		src_type = gl.FLOAT;
+	}
+	else if (ue_pixel_format === 'PF_G32R32F')
+	{
+		internal_format = gl.RG32F;
+		src_format = gl.RG;
+		src_type = gl.FLOAT;
+	}
+	else if (ue_pixel_format === 'PF_R32_FLOAT')
+	{
+		internal_format = gl.R32F;
+		src_format = gl.RED;
+		src_type = gl.FLOAT;
+	}
+	// 16bit floats
+	else if (ue_pixel_format === 'PF_FloatRGBA')
+	{
+		internal_format = gl.RGBA16F;
+		src_format = gl.RGBA;
+		src_type = gl.HALF_FLOAT;
+	}
+	else if (ue_pixel_format === 'PF_G16R16F')
+	{
+		internal_format = gl.RG16F;
+		src_format = gl.RG;
+		src_type = gl.HALF_FLOAT;
+	}
+	else if (ue_pixel_format === 'PF_R16F')
+	{
+		internal_format = gl.R16F;
+		src_format = gl.RED;
+		src_type = gl.HALF_FLOAT;
+	}
+	// 11-11-10bit floats
+	else if (ue_pixel_format === 'PF_FloatRGB' || ue_pixel_format === 'PF_FloatR11G11B10')
+	{
+		internal_format = gl.R11F_G11F_B10F;
+		src_format = gl.RGB;
+		src_type = gl.UNSIGNED_INT_10F_11F_11F_REV;
+	}
+	// 16bit unorms
+	else if (ue_pixel_format === 'PF_A16B16G16R16')
+	{
+		internal_format = g_gl_ext['EXT_texture_norm16'].RGBA16_EXT;
+		src_format = gl.RGBA;
+		src_type = gl.UNSIGNED_SHORT;
+	}
+	else if (ue_pixel_format === 'PF_G16R16')
+	{
+		internal_format = g_gl_ext['EXT_texture_norm16'].RG16_EXT;
+		src_format = gl.RG;
+		src_type = gl.UNSIGNED_SHORT;
+	}
+	else if (ue_pixel_format === 'PF_G16')
+	{
+		internal_format = g_gl_ext['EXT_texture_norm16'].R16_EXT;
+		src_format = gl.RED;
+		src_type = gl.UNSIGNED_SHORT;
+	}
+	// 8bit unorms
+	else if (ue_pixel_format === 'PF_B8G8R8A8' || ue_pixel_format === 'PF_R8G8B8A8')
+	{
+		internal_format = gl.RGBA;
+		src_format = gl.RGBA;
+		src_type = gl.UNSIGNED_BYTE;
+	}
+	else if (ue_pixel_format === 'PF_R8G8')
+	{
+		internal_format = gl.RG8;
+		src_format = gl.RG;
+		src_type = gl.UNSIGNED_BYTE;
+	}
+	else if (ue_pixel_format === 'PF_R8' || ue_pixel_format === 'PF_G8' || ue_pixel_format === 'PF_A8' || ue_pixel_format === 'PF_L8')
+	{
+		internal_format = gl.R8;
+		src_format = gl.RED;
+		src_type = gl.UNSIGNED_BYTE;
+	}
+	// 8bit uint
+	else if (ue_pixel_format === 'PF_R8G8B8A8_UINT')
+	{
+		internal_format = gl.RGBA8UI;
+		src_format = gl.RGBA_INTEGER;
+		src_type = gl.UNSIGNED_BYTE;
+	}
+	else if (ue_pixel_format === 'PF_R8_UINT')
+	{
+		internal_format = gl.R8UI;
+		src_format = gl.RED_INTEGER;
+		src_type = gl.UNSIGNED_BYTE;
+	}
+	else
+	{
+		return null;
+	}
+
+	if (src_format === gl.RED_INTEGER || src_format === gl.RG_INTEGER || src_format === gl.RGBA_INTEGER)
+	{
+		var divide = 255;
+		if (src_type == gl.UNSIGNED_SHORT)
+		{
+			divide = 65535;
+		}
+
+		sampler_type = 'usampler2D';
+		sampler_return_type = 'uvec4';
+		swizzling = `(vec4(${swizzling}) * (1.0 / ${divide}.0))`;
+	}
+
+	if (src_format === gl.RED)
+	{
+		swizzling = `vec4((${swizzling}).r, 0.0, 0.0, 1.0)`;
+	}
+	else if (src_format === gl.RG)
+	{
+		swizzling = `vec4((${swizzling}).rg, 0.0, 1.0)`;
+	}
+	else if (src_format === gl.RGB)
+	{
+		swizzling = `vec4((${swizzling}).rgb, 1.0)`;
+	}
+
+	var gl_format = {};
+	gl_format['internal_format'] = internal_format;
+	gl_format['src_format'] = src_format;
+	gl_format['src_type'] = src_type;
+	gl_format['swizzling'] = swizzling;
+	gl_format['sampler_type'] = sampler_type;
+	gl_format['sampler_return_type'] = sampler_return_type;
+	gl_format['ue_pixel_format'] = ue_pixel_format;
+	return gl_format;
+}
+
+function create_texture_from_pixel_data(gl, resource_desc, resource_version_name, raw_pixel_data)
+{
+	if (raw_pixel_data === null)
+	{
+		console.error(`failed to load ${resource_version_name}`);
+		return null;
+	}
+
+	var gl_format = translate_pixel_format(gl, resource_desc['Format']);
+	if (gl_format === null)
+	{
+		return null;
+	}
+
+	var converted_data = null;
+	if (gl_format['src_type'] === gl.UNSIGNED_BYTE)
+	{
+		converted_data = new Uint8Array(raw_pixel_data);
+	}
+	else if (gl_format['src_type'] === gl.UNSIGNED_SHORT || gl_format['src_type'] === gl.HALF_FLOAT)
+	{
+		converted_data = new Uint16Array(raw_pixel_data);
+	}
+	else if (gl_format['src_type'] === gl.UNSIGNED_INT_10F_11F_11F_REV)
+	{
+		converted_data = new Uint32Array(raw_pixel_data);
+	}
+	else if (gl_format['src_type'] === gl.FLOAT)
+	{
+		converted_data = new Float32Array(raw_pixel_data);
+	}
+	else
+	{
+		return null;
+	}
+
+	const texture = gl.createTexture();
+	gl.bindTexture(gl.TEXTURE_2D, texture);
+	gl.texImage2D(
+		gl.TEXTURE_2D,
+		/* level = */ 0,
+		gl_format['internal_format'],
+	    resource_desc['ExtentX'],
+	    resource_desc['ExtentY'],
+	    /* border = */ 0,
+	    gl_format['src_format'],
+	    gl_format['src_type'],
+	    converted_data);
+	gl.texParameteri(gl.TEXTURE_2D, gl.TEXTURE_MAG_FILTER, gl.NEAREST);
+	gl.texParameteri(gl.TEXTURE_2D, gl.TEXTURE_MIN_FILTER, gl.NEAREST);
+	gl.texParameteri(gl.TEXTURE_2D, gl.TEXTURE_WRAP_S, gl.CLAMP_TO_EDGE);
+	gl.texParameteri(gl.TEXTURE_2D, gl.TEXTURE_WRAP_T, gl.CLAMP_TO_EDGE);
+	gl.bindTexture(gl.TEXTURE_2D, null);
+
+	return texture;
+}
+
+
+function compile_screen_shader_program(gl, frag_code)
+{
+	const shader_code_prefix = `#version 300 es
+
+precision highp float;
+precision highp int;
+precision highp sampler2D;
+precision highp usampler2D;`;
+
+	var vert_code = `
+in vec3 coordinates;
+out highp vec2 tex_coordinates;
+void main(void) {
+	gl_Position = vec4(coordinates, 1.0);
+	tex_coordinates = coordinates.xy * 0.5 + 0.5;
+}`;
+
+	var vert_shader = gl.createShader(gl.VERTEX_SHADER);
+	{
+		gl.shaderSource(vert_shader, shader_code_prefix + vert_code);
+		gl.compileShader(vert_shader);
+
+		var compilation_status = gl.getShaderParameter(vert_shader, gl.COMPILE_STATUS);
+		if (!compilation_status)
+		{
+			var compilation_log = gl.getShaderInfoLog(vert_shader);
+			console.log('Fragment shader compiler log: ' + compilation_log);
+			return null;
+		}
+	}
+
+	var frag_shader = gl.createShader(gl.FRAGMENT_SHADER);
+	{
+		gl.shaderSource(frag_shader, shader_code_prefix + frag_code); 
+		gl.compileShader(frag_shader);
+
+		var compilation_status = gl.getShaderParameter(frag_shader, gl.COMPILE_STATUS);
+		if (!compilation_status)
+		{
+			var compilation_log = gl.getShaderInfoLog(frag_shader);
+			console.log('Fragment shader compiler log: ' + compilation_log);
+			document.getElementById('texture_visualization_code_log').value = 'Compilation failed:\n' + compilation_log;
+			return null;
+		}
+	}
+
+	var shader_program = gl.createProgram();
+	{
+		gl.attachShader(shader_program, vert_shader);
+		gl.attachShader(shader_program, frag_shader);
+		gl.linkProgram(shader_program);
+
+		var link_status = gl.getProgramParameter(shader_program, gl.LINK_STATUS);
+		if (!link_status)
+		{
+			var link_log = gl.getProgramInfoLog(shader_program);
+			console.log('Program compiler log: ' + link_log);
+			document.getElementById('texture_visualization_code_log').value = 'Link failed:\n' + link_log;
+		}
+		else
+		{
+			document.getElementById('texture_visualization_code_log').value = 'Compilation succeeded';
+		}
+	}
+	return shader_program;
+}
+
+function refresh_texture_visualization(resource_name, resource_version_name, raw_pixel_data)
+{
+	var resource_desc = load_resource_desc(resource_name);
+
+	var canvas = document.getElementById('texture_visualization_canvas');
+	var gl = g_gl;
+
+	if (!g_gl_texture)
+	{
+		g_gl_texture = create_texture_from_pixel_data(gl, resource_desc, resource_version_name, raw_pixel_data);
+	}
+
+	var texture = g_gl_texture;
+	if (!texture)
+	{
+		return;
+	}
+
+
+	/*======== Defining and storing the geometry ===========*/
+
+	var vertices = [
+		-1.0,+1.0,0.0,
+		-1.0,-1.0,0.0,
+		+1.0,-1.0,0.0,
+		+1.0,+1.0,0.0, 
+	];
+
+	var indices = [0, 1, 2, 2, 0, 3];
+
+	var vertex_buffer = gl.createBuffer();
+	{
+		gl.bindBuffer(gl.ARRAY_BUFFER, vertex_buffer);
+		gl.bufferData(gl.ARRAY_BUFFER, new Float32Array(vertices), gl.STATIC_DRAW);
+		gl.bindBuffer(gl.ARRAY_BUFFER, null);
+	}
+
+	var Index_Buffer = gl.createBuffer();
+	{
+		gl.bindBuffer(gl.ELEMENT_ARRAY_BUFFER, Index_Buffer);
+		gl.bufferData(gl.ELEMENT_ARRAY_BUFFER, new Uint16Array(indices), gl.STATIC_DRAW);
+		gl.bindBuffer(gl.ELEMENT_ARRAY_BUFFER, null);
+	}
+
+	/*================ Shaders ====================*/
+
+	var user_frag_code = document.getElementById('texture_visualization_code_input').value;
+	var shader_program = compile_screen_shader_program(gl, user_frag_code);
+
+	if (!shader_program)
+	{
+		return;
+	}
+
+	// Save the user code.
+	g_shader_code_dict[resource_version_name] = user_frag_code;
+
+	gl.clearColor(0.0, 0.0, 0.0, 1.0);
+	gl.clear(gl.COLOR_BUFFER_BIT);
+
+	{
+		gl.useProgram(shader_program);
+		gl.bindBuffer(gl.ARRAY_BUFFER, vertex_buffer);
+		gl.bindBuffer(gl.ELEMENT_ARRAY_BUFFER, Index_Buffer);
+
+		var coord = gl.getAttribLocation(shader_program, "coordinates");
+		gl.vertexAttribPointer(coord, 3, gl.FLOAT, false, 0, 0); 
+		gl.enableVertexAttribArray(coord);
+
+		var tex_sampler = gl.getUniformLocation(shader_program, 'tex_sampler');
+		gl.activeTexture(gl.TEXTURE0);
+		gl.bindTexture(gl.TEXTURE_2D, texture);
+		gl.uniform1i(tex_sampler, 0);
+
+		gl.viewport(0, 0, canvas.width, canvas.height);
+		gl.drawElements(gl.TRIANGLES, indices.length, gl.UNSIGNED_SHORT, 0);
+	}
+}
+
+
+// -------------------------------------------------------------------- DISPLAY BUFFER
+
+var g_buffer_name = null;
+var g_buffer_version_name = null;
+var g_raw_buffer = null;
+
+function display_buffer(resource_name, resource_version_name)
+{
+	var resource_desc = load_resource_desc(resource_name);
+
+	g_texture_name = resource_name;
+	g_texture_version_name = resource_version_name;
+	g_gl = null;
+	g_gl_texture = null;
+	g_raw_buffer = null;
+	g_buffer_name = null;
+	g_buffer_version_name = null;
+
+	var default_format = '';
+	{
+		var byte_per_element = resource_desc['BytesPerElement'];
+		var column_count = byte_per_element / 4;
+
+		var format_per_column = [];
+		for (var i = 0; i < column_count; i++)
+		{
+			format_per_column.push('hex');
+		}
+
+		default_format = format_per_column.join(', ');
+	}
+
+	var resource_info_htmls = '';
+	{
+		resource_info_htmls += `
+			<table width="100%" class="resource_desc">`;
+		
+		for (var key in resource_desc){
+		    if (resource_desc.hasOwnProperty(key)){
+				resource_info_htmls += `
+					<tr>
+						<td>${key}</td>
+						<td>${resource_desc[key]}</td>
+					</tr>`;
+		    }
+		}
+
+		resource_info_htmls += `
+			</table>`;
+	}
+
+	var html = '';
+	html += `
+		<div class="main_div">
+			<div class="selection_list_title">Buffer descriptor: ${resource_desc['Name']}</div>
+			<div>
+				${resource_info_htmls}
+			</div>
+		</div>
+		<div class="main_div">
+			<div class="selection_list_title" style="width: auto;">Buffer visualization: ${resource_desc['Name']}</div>
+			<div><input type="text" id="buffer_visualization_format" value="${default_format}" onchange="refresh_current_buffer_visualization();" style="width: 100%;" /></div>
+			<div id="buffer_content_pannel"></div>
+		</div>`;
+
+	document.getElementById('display_resource_pannel').innerHTML = html + get_resource_producer_html(resource_name);
+
+	load_binary_file(`Resources/${resource_version_name}.bin`, function(raw_buffer_data)
+	{
+		if (raw_buffer_data)
+		{
+			g_raw_buffer = raw_buffer_data;
+			g_buffer_name = resource_name;
+			g_buffer_version_name = resource_version_name;
+			refresh_buffer_visualization(resource_name, resource_version_name, raw_buffer_data);
+		}
+		else
+		{
+			document.getElementById('buffer_content_pannel').innerHTML = `
+				Error: Resources/${resource_version_name}.bin couldn't be loaded.`;
+		}
+	});
+}
+
+function refresh_current_buffer_visualization()
+{
+	refresh_buffer_visualization(g_buffer_name, g_buffer_version_name, g_raw_buffer);
+}
+
+function refresh_buffer_visualization(resource_name, buffer_visualization_format, raw_buffer_data)
+{
+	var resource_desc = load_resource_desc(resource_name);
+	var byte_per_element = resource_desc['BytesPerElement'];
+	var num_element = resource_desc['NumElements'];
+
+	var column_count = byte_per_element / 4;
+
+	var column_byte_offset = [];
+	var buffer_format = [];
+	{
+		var text_format = document.getElementById('buffer_visualization_format').value;
+
+		var column_formats = text_format.split(',');
+
+		var byte_offset = 0;
+
+		column_formats.forEach(function(value, index) {
+			var format_to_use = value.trim();
+
+			column_byte_offset.push(byte_offset);
+			buffer_format.push(format_to_use);
+
+			byte_offset += 4;
+		});
+
+		for (var i = buffer_format.length; i < column_count; i++)
+		{
+			var format_to_use = buffer_format[i % column_formats.length];
+
+			column_byte_offset.push(byte_offset);
+			buffer_format.push(format_to_use);
+
+			byte_offset += 4;
+		}
+
+	}
+
+	var resource_content = '';
+	{
+		resource_content += `<table id="buffer_content_table">
+			<tr>
+				<td>Address</td>`;
+
+		for (var i = 0; i < column_count; i++)
+		{
+			resource_content += `<td>${buffer_format[i]}</td>`;
+		}
+
+		resource_content += `</tr>`;
+
+		for (var elem_id = 0; elem_id < num_element; elem_id++)
+		{
+			resource_content += `
+				<tr>
+					<td>${elem_id * byte_per_element}</td>`;
+
+			for (var i = 0; i < column_count; i++)
+			{
+				var byte_offset = byte_per_element * elem_id + column_byte_offset[i];
+
+				var value = null;
+				var display = buffer_format[i];
+				if (display == 'float')
+				{
+					value = (new Float32Array(raw_buffer_data, byte_offset, 1))[0];
+				}
+				else if (display == 'int')
+				{
+					value = (new Int32Array(raw_buffer_data, byte_offset, 1))[0];
+				}
+				else if (display == 'uint')
+				{
+					value = (new Uint32Array(raw_buffer_data, byte_offset, 1))[0];
+				}
+				else if (display == 'hex')
+				{
+					value = (new Uint32Array(raw_buffer_data, byte_offset, 1))[0].toString(16);;
+				}
+				else if (display == 'ishort')
+				{
+					value = (new Int16Array(raw_buffer_data, byte_offset, 1))[0];
+				}
+				else if (display == 'ushort')
+				{
+					value = (new Uint16Array(raw_buffer_data, byte_offset, 1))[0];
+				}
+				else if (display == 'char')
+				{
+					value = (new Int8Array(raw_buffer_data, byte_offset, 1))[0];
+				}
+				else if (display == 'uchar')
+				{
+					value = (new Uint8Array(raw_buffer_data, byte_offset, 1))[0];
+				}
+				else
+				{
+					value = `Unknown ${display}`;
+				}
+
+				resource_content += `<td>${value}</td>`;
+			}
+
+			resource_content += `</tr>`;
+		}
+
+		resource_content += `</table>`;
+	}
+
+	document.getElementById('buffer_content_pannel').innerHTML = resource_content;
+}
+
+
+
+// -------------------------------------------------------------------- WINDOW LAYOUT
+
+function init_top()
+{
+	console.log(g_infos);
+	document.getElementById('top_bar').innerHTML = `${g_infos['Project']} - ${g_infos['Platform']} - ${g_infos['RHI']} - ${g_infos['BuildVersion']} - ${g_infos['DumpTime']} `;
+}
+
+function init_page()
+{
+	g_infos = load_json('Infos.json');
+	g_passes = load_json_dict_sequence('Passes.json');
+	console.log(g_passes[0]);
+	init_top();
+	display_pass_hierarchy();
+}
+
+function onresize_body()
+{
+	var w = window.innerWidth;
+	var h = window.innerHeight;
+
+	var top_h = 40;
+
+	document.getElementById('onresize_body_top_bar').style.width = `${w}px`;
+	document.getElementById('onresize_body_top_bar').style.height = `${top_h}px`;
+
+	document.getElementById('onresize_body_left').style.width = `${Math.floor(w * 0.25)}px`;
+	document.getElementById('onresize_body_left').style.height = `${h - top_h}px`;
+
+	document.getElementById('onresize_body_right').style.width = `${w - Math.floor(w * 0.25)}px`;
+	document.getElementById('onresize_body_right').style.height = `${h - top_h}px`;
+
+	document.getElementById('onresize_body_right').style.width = `${w - Math.floor(w * 0.25)}px`;
+	document.getElementById('onresize_body_right').style.height = `${h - top_h}px`;
+
+	{
+		var search_h = (30 + 2 * 5) * 2;
+		document.getElementById('pass_hierarchy').style.height = `${h - top_h - search_h - 10 * 2}px`;
+	}
+
+	{
+		document.getElementById('main_right_pannel').style.height = `${h - top_h - 10 * 2}px`;
+	}
+}
+
+</script>
+
+
+<style type="text/css">
+
+body
+{
+	background: #111111;
+	color: #DDDDDD;
+	margin: 0;
+	padding: 0;
+	/*font-family: Arial, Helvetica, sans-serif;*/
+	font-family: consolas, "Liberation Mono", courier, monospace;
+	/* overflow: hidden; */
+}
+
+table, tr, td, div
+{
+	margin: 0;
+	padding: 0;
+	border-spacing: 0;
+	border-collapse: collapse;
+	vertical-align: top;
+}
+
+td
+{
+	display: inline-block;
+}
+
+div
+{
+	display: block;
+}
+
+.main_div
+{
+	background: #222222;
+	padding: 5px;
+	margin: 5px;
+}
+
+#main_right_pannel .main_div
+{
+	margin-bottom: 20px;
+}
+
+
+/* ----------------------------------------------------- common layout */
+
+.selection_list
+{
+	max-height: inherit;
+	overflow-x: auto;
+	overflow-y: scroll;
+}
+
+.selection_list div
+{
+	width: auto;
+	padding: 5px 10px;
+	white-space: nowrap;
+}
+
+.selection_list div:nth-child(odd)
+{
+	background: #171717;
+}
+
+.selection_list div.selected
+{
+	background: #114422;
+}
+
+.selection_list div:hover
+{
+	background: #112244;
+	cursor: pointer;
+}
+
+.selection_list div.disabled
+{
+	color: #999999;
+}
+
+#main_right_pannel .selection_list
+{
+	max-height: 20%;
+}
+
+.selection_list_title
+{
+	font-size: 20;
+	font-weight: bold;
+	padding: 5px 20px 10px 20px;
+}
+
+.resource_desc tr:nth-child(odd)
+{
+	background: #171717;
+}
+
+.resource_desc tr td
+{
+	padding: 5px 20px;
+}
+
+.resource_desc tr td
+{
+	text-align: right;
+}
+
+.resource_desc tr td:first-child
+{
+	width: 150px;
+}
+
+
+
+#top_bar
+{
+	font-size: 20;
+	font-weight: bold;
+	padding: 10px;
+}
+
+#main_right_pannel
+{
+	width: auto;
+	overflow-x: auto;
+	overflow-y: scroll;
+}
+
+#main_right_pannel .pass_title
+{
+	font-size: 20;
+	font-weight: bolder;
+	padding: 10px 40px;
+}
+
+#texture_visualization_code_input
+{
+	min-width: 800px;
+	min-height: 200px;
+	display: block;
+	width: auto;
+}
+
+#texture_visualization_code_log
+{
+	margin-top: 10px;
+	min-width: 800px;
+	min-height: 200px;
+	display: block;
+	width: auto;
+}
+
+#buffer_content_table tr:nth-child(even)
+{
+	background: #171717;
+}
+
+#buffer_content_table tr td
+{
+	padding: 3px 5px;
+	width: 40px;
+	height: 20px;
+	overflow: hidden;
+	display: table-cell;
+	text-align: right;
+	font-family: consolas, "Liberation Mono", courier, monospace;
+}
+
+canvas
+{
+	background-color: black;
+	border: 1px solid #333333;
+}
+
+</style>
+
+
+<body onload="onresize_body(); init_page();" onresize="onresize_body();">
+	<table cellpadding="0" cellspacing="0" border="0" id="onresize_body_table">
+		<tr>
+			<td colspan="2" id="onresize_body_top_bar">
+				<div id="top_bar"></div>
+			</td>
+		</tr>
+		<tr>
+			<td valign="top" id="onresize_body_left">
+				<div class="main_div">
+					<div id="pass_hierarchy_search" style="padding: 5px; height: 70px;">
+						<input type="text" id="pass_search_input" oninput="display_pass_hierarchy();" onchange="display_pass_hierarchy();" style="width: 100%;" placeholder="Search pass..." />
+						<input type="text" id="resource_search_input" oninput="display_pass_hierarchy();" onchange="display_pass_hierarchy();" style="width: 100%; margin-top: 10px;" placeholder="Search resource..." />
+					</div>
+					<div id="pass_hierarchy" class="selection_list"></div>
+				</div>
+			</td>
+			<td valign="top" id="onresize_body_right">
+				<div id="main_right_pannel"></div>
+			</td>
+		</tr>
+	</table>
+</body>
+</html>
diff --git a/Engine/Extras/GPUDumpViewer/OpenGPUDumpViewer.bat b/Engine/Extras/GPUDumpViewer/OpenGPUDumpViewer.bat
new file mode 100755
--- /dev/null
+++ b/Engine/Extras/GPUDumpViewer/OpenGPUDumpViewer.bat
@@ -0,0 +1,5 @@
+set CWD=%cd%
+
+REM --allow-file-access-from-files allow to load a file from a file:// webpage required for GPUDumpViewer.html to work.
+REM --user-data-dir is required to force chrome to open a new instance so that --allow-file-access-from-files is honored.
+"C:\Program Files\Google\Chrome\Application\chrome.exe" file://%CWD%/GPUDumpViewer.html --allow-file-access-from-files --new-window --incognito --user-data-dir=%CWD%/.tmp_chrome_data/
diff --git a/Engine/Plugins/RenderGraphInsights/Source/RenderGraphInsights/Private/RenderGraphTrack.cpp b/Engine/Plugins/RenderGraphInsights/Source/RenderGraphInsights/Private/RenderGraphTrack.cpp
--- a/Engine/Plugins/RenderGraphInsights/Source/RenderGraphInsights/Private/RenderGraphTrack.cpp
+++ b/Engine/Plugins/RenderGraphInsights/Source/RenderGraphInsights/Private/RenderGraphTrack.cpp
@@ -1524,24 +1524,6 @@ TSharedPtr<ITimingEventFilter> FRenderGraphTrack::GetFilterByEvent(const TShared
 	return nullptr;
 }
 
-const TCHAR* GetDimensionName(ETextureDimension Dimension)
-{
-	switch (Dimension)
-	{
-	case ETextureDimension::Texture2D:
-		return TEXT("Texture2D");
-	case ETextureDimension::Texture2DArray:
-		return TEXT("Texture2DArray");
-	case ETextureDimension::Texture3D:
-		return TEXT("Texture3D");
-	case ETextureDimension::TextureCube:
-		return TEXT("TextureCube");
-	case ETextureDimension::TextureCubeArray:
-		return TEXT("TextureCubeArray");
-	}
-	return TEXT("");
-}
-
 void FRenderGraphTrack::BuildContextMenu(FMenuBuilder& MenuBuilder)
 {
 	Super::BuildContextMenu(MenuBuilder);
@@ -2080,7 +2062,7 @@ void FRenderGraphTrack::InitTooltip(FTooltipDrawState& Tooltip, const ITimingEve
 			const FTexturePacket& Texture = TooltipEvent.GetPacket();
 
 			Tooltip.AddTitle(Texture.Name);
-			Tooltip.AddNameValueTextLine(TEXT("Dimension:"), GetDimensionName(Texture.Desc.Dimension));
+			Tooltip.AddNameValueTextLine(TEXT("Dimension:"), GetTextureDimensionString(Texture.Desc.Dimension));
 			Tooltip.AddNameValueTextLine(TEXT("Create Flags:"), GetTextureCreateFlagsName(Texture.Desc.Flags));
 			Tooltip.AddNameValueTextLine(TEXT("Format:"), UEnum::GetValueAsString(Texture.Desc.Format));
 			Tooltip.AddNameValueTextLine(TEXT("Extent:"), FString::Printf(TEXT("%d, %d"), Texture.Desc.Extent.X, Texture.Desc.Extent.Y));
diff --git a/Engine/Source/Runtime/Launch/Private/LaunchEngineLoop.cpp b/Engine/Source/Runtime/Launch/Private/LaunchEngineLoop.cpp
--- a/Engine/Source/Runtime/Launch/Private/LaunchEngineLoop.cpp
+++ b/Engine/Source/Runtime/Launch/Private/LaunchEngineLoop.cpp
@@ -173,6 +173,7 @@
 	#include "RenderUtils.h"
 	#include "DynamicResolutionState.h"
 	#include "EngineModule.h"
+	#include "RenderGraphBuilder.h"
 
 #if !UE_SERVER
 	#include "AppMediaTimeSource.h"
@@ -5505,6 +5506,17 @@ void FEngineLoop::Tick()
 
 		FCoreDelegates::OnEndFrame.Broadcast();
 
+		// end of RDG resource dump
+		#if WITH_ENGINE
+			#if RDG_DUMP_RESOURCES
+				ENQUEUE_RENDER_COMMAND(EndRDGResourceDump)(
+					[](FRHICommandListImmediate& RHICmdList)
+					{
+						FRDGBuilder::EndResourceDump();
+					});
+			#endif
+		#endif
+
 		#if !UE_SERVER && WITH_ENGINE
 		{
 			// We emit dynamic resolution's end frame right before RHI's. GEngine is going to ignore it if no BeginFrame was done.
diff --git a/Engine/Source/Runtime/RHI/Public/RHIDefinitions.h b/Engine/Source/Runtime/RHI/Public/RHIDefinitions.h
--- a/Engine/Source/Runtime/RHI/Public/RHIDefinitions.h
+++ b/Engine/Source/Runtime/RHI/Public/RHIDefinitions.h
@@ -2260,6 +2260,24 @@ inline const TCHAR* GetShaderFrequencyString(EShaderFrequency Frequency, bool bI
 	return String;
 };
 
+inline const TCHAR* GetTextureDimensionString(ETextureDimension Dimension)
+{
+	switch (Dimension)
+	{
+	case ETextureDimension::Texture2D:
+		return TEXT("Texture2D");
+	case ETextureDimension::Texture2DArray:
+		return TEXT("Texture2DArray");
+	case ETextureDimension::Texture3D:
+		return TEXT("Texture3D");
+	case ETextureDimension::TextureCube:
+		return TEXT("TextureCube");
+	case ETextureDimension::TextureCubeArray:
+		return TEXT("TextureCubeArray");
+	}
+	return TEXT("");
+}
+
 inline bool IsGeometryPipelineShaderFrequency(EShaderFrequency Frequency)
 {
 	return Frequency == SF_Mesh || Frequency == SF_Amplification;
diff --git a/Engine/Source/Runtime/RenderCore/Private/RenderGraphBuilder.cpp b/Engine/Source/Runtime/RenderCore/Private/RenderGraphBuilder.cpp
--- a/Engine/Source/Runtime/RenderCore/Private/RenderGraphBuilder.cpp
+++ b/Engine/Source/Runtime/RenderCore/Private/RenderGraphBuilder.cpp
@@ -1486,7 +1486,11 @@ void FRDGBuilder::Execute(EExecuteMode ExecuteMode)
 	SCOPED_NAMED_EVENT_TEXT("FRDGBuilder::Execute", FColor::Magenta);
 
 	// Create the epilogue pass at the end of the graph just prior to compilation.
-	SetupEmptyPass(EpiloguePass = Passes.Allocate<FRDGSentinelPass>(Allocator, RDG_EVENT_NAME("Graph Epilogue%s", (ExecuteMode == EExecuteMode::Drain) ? TEXT(" (Drain)") : TEXT(""))));
+	{
+		bInDebugPassScope = true;
+		SetupEmptyPass(EpiloguePass = Passes.Allocate<FRDGSentinelPass>(Allocator, RDG_EVENT_NAME("Graph Epilogue%s", (ExecuteMode == EExecuteMode::Drain) ? TEXT(" (Drain)") : TEXT(""))));
+		bInDebugPassScope = false;
+	}
 
 	const FRDGPassHandle ProloguePassHandle = GetProloguePassHandle();
 	const FRDGPassHandle EpiloguePassHandle = GetEpiloguePassHandle();
@@ -2127,6 +2131,10 @@ void FRDGBuilder::SetupPassInternal(FRDGPass* Pass, FRDGPassHandle PassHandle, E
 	}
 
 	IF_RDG_ENABLE_DEBUG(VisualizePassOutputs(Pass));
+
+	#if RDG_DUMP_RESOURCES
+		DumpResourcePassOutputs( Pass);
+	#endif
 }
 
 void FRDGBuilder::SetupBufferUploads()
@@ -2405,6 +2413,8 @@ void FRDGBuilder::CreateUniformBuffers()
 
 void FRDGBuilder::AddProloguePass()
 {
+	bInDebugPassScope = true;
+
 	if (AsyncComputePassCount > 0)
 	{
 		FRDGPass* AsyncComputePass = SetupEmptyPass(Passes.Allocate<FRDGSentinelPass>(Allocator, RDG_EVENT_NAME("Graph Prologue (AsyncCompute)"), ERDGPassFlags::AsyncCompute));
@@ -2421,6 +2431,8 @@ void FRDGBuilder::AddProloguePass()
 	{
 		ProloguePassHandles[ERHIPipeline::AsyncCompute] = GraphicsPass->Handle;
 	}
+
+	bInDebugPassScope = false;
 }
 
 void FRDGBuilder::ExecutePassPrologue(FRHIComputeCommandList& RHICmdListPass, FRDGPass* Pass)
@@ -3353,6 +3365,15 @@ void FRDGBuilder::VisualizePassOutputs(const FRDGPass* Pass)
 		return;
 	}
 
+	// Don't scan any of the pass output if not texture is requested to be visualized.
+	if (GVisualizeTexture.Requested.Name.IsEmpty())
+	{
+		return;
+	}
+
+	bInDebugPassScope = true;
+
+
 	Pass->GetParameters().EnumerateTextures([&](FRDGParameter Parameter)
 	{
 		switch (Parameter.GetType())
@@ -3417,6 +3438,8 @@ void FRDGBuilder::VisualizePassOutputs(const FRDGPass* Pass)
 		break;
 		}
 	});
+
+	bInDebugPassScope = false;
 #endif
 }
 
diff --git a/Engine/Source/Runtime/RenderCore/Private/RenderGraphDumpResources.cpp b/Engine/Source/Runtime/RenderCore/Private/RenderGraphDumpResources.cpp
new file mode 100644
--- /dev/null
+++ b/Engine/Source/Runtime/RenderCore/Private/RenderGraphDumpResources.cpp
@@ -0,0 +1,581 @@
+// Copyright Epic Games, Inc. All Rights Reserved.
+
+#include "RenderGraph.h"
+#include "HAL/PlatformFileManager.h"
+#include "Misc/FileHelper.h"
+#include "Misc/WildcardString.h"
+#include "BuildSettings.h"
+#include "Serialization/JsonTypes.h"
+#include "Serialization/JsonReader.h"
+#include "Serialization/JsonSerializer.h"
+
+#if RDG_DUMP_RESOURCES
+
+namespace
+{
+
+static TAutoConsoleVariable<FString> GDumpGPURootCVar(
+	TEXT("r.DumpGPU.Root"),
+	TEXT("*"),
+	TEXT("Allows to filter the tree when using r.DumpGPU command, the pattern match is case sensitive."),
+	ECVF_Default);
+
+BEGIN_SHADER_PARAMETER_STRUCT(FDumpTexturePass, )
+	RDG_TEXTURE_ACCESS(Texture, ERHIAccess::CopySrc)
+END_SHADER_PARAMETER_STRUCT()
+
+BEGIN_SHADER_PARAMETER_STRUCT(FDumpBufferPass, )
+	RDG_BUFFER_ACCESS(Buffer, ERHIAccess::CopySrc)
+END_SHADER_PARAMETER_STRUCT()
+
+struct FRDGResourceDumpContext
+{
+	FString DumpingDirectoryPath;
+	FDateTime Time;
+	int32 ResourcesDumped = 0;
+	int32 PassesCount = 0;
+
+	bool bShowInExplore = false;
+
+	bool IsDumpingFrame() const
+	{
+		check(IsInRenderingThread());
+		return !DumpingDirectoryPath.IsEmpty();
+	}
+
+	static FString GetResourceName(const FRDGResource* Resource)
+	{
+		return FString::Printf(TEXT("%s.%p"), Resource->Name, Resource);
+	}
+
+	static FString GetResourceDumpName(const FRDGPass* Pass, const FRDGResource* Resource)
+	{
+		return FString::Printf(TEXT("%s.%p.v%p.bin"), Resource->Name, Resource, Pass);
+	}
+
+	FString GetResourceDumpDirectory() const
+	{
+		return DumpingDirectoryPath + TEXT("Resources/");
+	}
+
+	void AddDumpTexturePass(
+		FRDGBuilder& GraphBuilder,
+		TArray<TSharedPtr<FJsonValue>>& InputResourceNames,
+		TArray<TSharedPtr<FJsonValue>>& OutputResourceNames,
+		const FRDGPass* Pass,
+		FRDGTexture* Texture,
+		bool bIsOutputResource)
+	{
+		FString UniqueResourceName = GetResourceName(Texture);
+
+		if (bIsOutputResource)
+		{
+			OutputResourceNames.AddUnique(MakeShareable(new FJsonValueString(UniqueResourceName)));
+		}
+		else
+		{
+			InputResourceNames.AddUnique(MakeShareable(new FJsonValueString(UniqueResourceName)));
+		}
+
+		const FRDGTextureDesc& Desc = Texture->Desc;
+		int32 ByteSize = 0;
+		{
+			bool bIsUnsupported = false;
+
+			if (GPixelFormats[Desc.Format].BlockSizeX != 1 || 
+				GPixelFormats[Desc.Format].BlockSizeY != 1 ||
+				GPixelFormats[Desc.Format].BlockSizeZ != 1)
+			{
+				bIsUnsupported = true;
+			}
+
+			if (Desc.Format == PF_DepthStencil)
+			{
+				bIsUnsupported = true;
+			}
+
+			if (!bIsUnsupported && Desc.IsTexture2D() && !Desc.IsMultisample() && !Desc.IsTextureArray())
+			{
+				ByteSize = Desc.Extent.X * Desc.Extent.Y * GPixelFormats[Desc.Format].BlockBytes;
+			}
+		}
+
+		FString ResourceDumpDirectory = GetResourceDumpDirectory();
+
+		// Dump the information of the texture to json file.
+		{
+			FString PixelFormat = FString::Printf(TEXT("PF_%s"), GPixelFormats[Desc.Format].Name);
+
+			TSharedPtr<FJsonObject> JsonObject = MakeShareable(new FJsonObject);
+			JsonObject->SetStringField(TEXT("Name"), Texture->Name);
+			JsonObject->SetNumberField(TEXT("ByteSize"), ByteSize);
+			JsonObject->SetStringField(TEXT("Type"), GetTextureDimensionString(Desc.Dimension));
+			JsonObject->SetStringField(TEXT("Format"), *PixelFormat);
+			JsonObject->SetNumberField(TEXT("ExtentX"), Desc.Extent.X);
+			JsonObject->SetNumberField(TEXT("ExtentY"), Desc.Extent.Y);
+			JsonObject->SetNumberField(TEXT("Depth"), Desc.Depth);
+			JsonObject->SetNumberField(TEXT("ArraySize"), Desc.ArraySize);
+			JsonObject->SetNumberField(TEXT("NumMips"), Desc.NumMips);
+			JsonObject->SetNumberField(TEXT("NumSamples"), Desc.NumSamples);
+
+			FString OutputString;
+			TSharedRef<TJsonWriter<TCHAR, TPrettyJsonPrintPolicy<TCHAR>>> Writer = TJsonWriterFactory<TCHAR, TPrettyJsonPrintPolicy<TCHAR>>::Create(&OutputString);
+			FJsonSerializer::Serialize(JsonObject.ToSharedRef(), Writer);
+
+			FString ResourceInfoJsonPath = ResourceDumpDirectory + UniqueResourceName + TEXT(".json");
+			FFileHelper::SaveStringToFile(OutputString, *ResourceInfoJsonPath);
+		}
+
+		if (ByteSize == 0)
+		{
+			return;
+		}
+
+		// Dump the resource's binary to a .bin file.
+		if (bIsOutputResource)
+		{
+			FString DumpFilePath = ResourceDumpDirectory + GetResourceDumpName(Pass, Texture);
+
+			FDumpTexturePass* PassParameters = GraphBuilder.AllocParameters<FDumpTexturePass>();
+			PassParameters->Texture = Texture;
+
+			GraphBuilder.AddPass(
+				RDG_EVENT_NAME("RDG DumpTexture(%s -> %s)", Texture->Name, *DumpFilePath),
+				PassParameters,
+				ERDGPassFlags::Readback,
+				[DumpFilePath, Texture, ByteSize](FRHICommandListImmediate& RHICmdList)
+			{
+				check(IsInRenderingThread());
+				TUniquePtr<FRHIGPUTextureReadback> Readback = MakeUnique<FRHIGPUTextureReadback>(TEXT("RDGTextureDumpReadback"));;
+
+				Readback->EnqueueCopyRDG(RHICmdList, Texture->GetRHI(), FResolveRect());
+
+				int32 RowPitchInPixels = 0;
+				void* Content = nullptr;
+				Readback->LockTexture(RHICmdList, /* out */ Content, /* out */ RowPitchInPixels);
+				if (Content)
+				{
+					TArray<uint8> Array;
+					Array.SetNumUninitialized(ByteSize);
+
+					int32 BytePerPixel = GPixelFormats[Texture->Desc.Format].BlockBytes;
+
+					const uint8* SrcData = reinterpret_cast<const uint8*>(Content);
+
+					for (int32 y = 0; y < Texture->Desc.Extent.Y; y++)
+					{
+						// Flip the data to be bottom left corner for the WebGL viewer.
+						const uint8* SrcPos = SrcData + (Texture->Desc.Extent.Y - 1 - y) * RowPitchInPixels * BytePerPixel;
+						uint8* DstPos = (&Array[0]) + y * Texture->Desc.Extent.X * BytePerPixel;
+
+						FPlatformMemory::Memmove(DstPos, SrcPos, Texture->Desc.Extent.X * BytePerPixel);
+					}
+
+					TArrayView<const uint8> ArrayView(reinterpret_cast<const uint8*>(&Array[0]), ByteSize);
+					FFileHelper::SaveArrayToFile(ArrayView, *DumpFilePath);
+
+					Readback->Unlock();
+				}
+			});
+		}
+
+		ResourcesDumped++;
+	}
+
+	void AddDumpBufferPass(
+		FRDGBuilder& GraphBuilder,
+		TArray<TSharedPtr<FJsonValue>>& InputResourceNames,
+		TArray<TSharedPtr<FJsonValue>>& OutputResourceNames,
+		const FRDGPass* Pass,
+		FRDGBuffer* Buffer,
+		bool bIsOutputResource)
+	{
+		FString UniqueResourceName = GetResourceName(Buffer);
+
+		if (bIsOutputResource)
+		{
+			OutputResourceNames.AddUnique(MakeShareable(new FJsonValueString(UniqueResourceName)));
+		}
+		else
+		{
+			InputResourceNames.AddUnique(MakeShareable(new FJsonValueString(UniqueResourceName)));
+		}
+
+		const FRDGBufferDesc& Desc = Buffer->Desc;
+		int32 ByteSize = Desc.GetTotalNumBytes();
+
+		FString ResourceDumpDirectory = GetResourceDumpDirectory();
+
+		// Dump the information of the buffer to json file.
+		{
+			TSharedPtr<FJsonObject> JsonObject = MakeShareable(new FJsonObject);
+			JsonObject->SetStringField(TEXT("Name"), Buffer->Name);
+			JsonObject->SetNumberField(TEXT("ByteSize"), ByteSize);
+			JsonObject->SetStringField(TEXT("Type"), GetBufferUnderlyingTypeName(Desc.UnderlyingType));
+			JsonObject->SetNumberField(TEXT("BytesPerElement"), Desc.BytesPerElement);
+			JsonObject->SetNumberField(TEXT("NumElements"), Desc.NumElements);
+
+			FString OutputString;
+			TSharedRef<TJsonWriter<TCHAR, TPrettyJsonPrintPolicy<TCHAR>>> Writer = TJsonWriterFactory<TCHAR, TPrettyJsonPrintPolicy<TCHAR>>::Create(&OutputString);
+			FJsonSerializer::Serialize(JsonObject.ToSharedRef(), Writer);
+
+			FString ResourceInfoJsonPath = ResourceDumpDirectory + UniqueResourceName + TEXT(".json");
+			FFileHelper::SaveStringToFile(OutputString, *ResourceInfoJsonPath);
+		}
+
+		if (Desc.UnderlyingType == FRDGBufferDesc::EUnderlyingType::AccelerationStructure)
+		{
+			return;
+		}
+
+		// Dump the resource's binary to a .bin file.
+		if (bIsOutputResource)
+		{
+			FString DumpFilePath = ResourceDumpDirectory + GetResourceDumpName(Pass, Buffer);
+
+			FDumpBufferPass* PassParameters = GraphBuilder.AllocParameters<FDumpBufferPass>();
+			PassParameters->Buffer = Buffer;
+
+			GraphBuilder.AddPass(
+				RDG_EVENT_NAME("RDG DumpBuffer(%s -> %s)", Buffer->Name, *DumpFilePath),
+				PassParameters,
+				ERDGPassFlags::Readback,
+				[DumpFilePath, Buffer, ByteSize](FRHICommandListImmediate& RHICmdList)
+			{
+				check(IsInRenderingThread());
+				TUniquePtr<FRHIGPUBufferReadback> Readback = MakeUnique<FRHIGPUBufferReadback>(TEXT("RDGBufferDumpReadback"));;
+
+				// TODO: work arround BUF_SourceCopy ensure()
+				Readback->EnqueueCopy(RHICmdList, Buffer->GetRHI(), ByteSize);
+
+				void* Content = Readback->Lock(ByteSize);
+				if (Content)
+				{
+					TArrayView<const uint8> ArrayView(reinterpret_cast<const uint8*>(Content), ByteSize);
+					FFileHelper::SaveArrayToFile(ArrayView, *DumpFilePath);
+
+					Readback->Unlock();
+				}
+				else
+				{
+					ensureMsgf(false, TEXT("%s couldn't be dumped"), Buffer->Name);
+				}
+			});
+		}
+
+		ResourcesDumped++;
+	}
+};
+
+FRDGResourceDumpContext GRDGResourceDumpContext;
+
+static FAutoConsoleCommand GDumpGPUCommand(
+	TEXT("r.DumpGPU"),
+	TEXT("Dump a multi-platform GPU capture of RDG resources, to be used when creating a JIRA of an artifact/rendering bug."),
+	FConsoleCommandWithArgsDelegate::CreateStatic(
+		[](const TArray<FString>& Args)
+		{
+			check(IsInGameThread());
+
+			ENQUEUE_RENDER_COMMAND(FSceneRendererCleanUp)(
+				[Args](FRHICommandListImmediate& RHICmdList)
+			{
+				FRDGBuilder::BeginResourceDump(Args);
+			});
+		}
+	)
+);
+
+}
+
+void FRDGBuilder::BeginResourceDump(const TArray<FString>& Args)
+{
+	if (GRDGResourceDumpContext.IsDumpingFrame())
+	{
+		return;
+	}
+
+	GRDGResourceDumpContext.Time = FDateTime::Now();
+	GRDGResourceDumpContext.DumpingDirectoryPath = FPaths::ProjectSavedDir() + TEXT("GPUDumps/") + FPlatformProperties::PlatformName() + TEXT("/") + FApp::GetProjectName() + TEXT("-") + GRDGResourceDumpContext.Time.ToString() + TEXT("/");
+
+	GRDGResourceDumpContext.bShowInExplore = Args.Contains(TEXT("-show"));
+
+	IPlatformFile& PlatformFile = FPlatformFileManager::Get().GetPlatformFile();
+	if (!PlatformFile.DirectoryExists(*GRDGResourceDumpContext.DumpingDirectoryPath))
+	{
+		PlatformFile.CreateDirectoryTree(*GRDGResourceDumpContext.DumpingDirectoryPath);
+	}
+	PlatformFile.CreateDirectoryTree(*(GRDGResourceDumpContext.DumpingDirectoryPath + TEXT("Resources/")));
+
+	// Output informations
+	{
+		TSharedPtr<FJsonObject> JsonObject = MakeShareable(new FJsonObject);
+		JsonObject->SetStringField(TEXT("Project"), FApp::GetProjectName());
+		JsonObject->SetStringField(TEXT("BuildBranch"), BuildSettings::GetBranchName());
+		JsonObject->SetStringField(TEXT("BuildDate"), BuildSettings::GetBuildDate());
+		JsonObject->SetStringField(TEXT("BuildVersion"), BuildSettings::GetBuildVersion());
+		JsonObject->SetStringField(TEXT("DumpTime"), GRDGResourceDumpContext.Time.ToString());
+		JsonObject->SetStringField(TEXT("Platform"), FPlatformProperties::IniPlatformName());
+		JsonObject->SetStringField(TEXT("RHI"), FApp::GetGraphicsRHI());
+		JsonObject->SetStringField(TEXT("GPUVendor"), RHIVendorIdToString());
+		JsonObject->SetStringField(TEXT("ComputerName"), FPlatformProcess::ComputerName());
+
+		FString OutputString;
+		TSharedRef<TJsonWriter<TCHAR, TPrettyJsonPrintPolicy<TCHAR>>> Writer = TJsonWriterFactory<TCHAR, TPrettyJsonPrintPolicy<TCHAR>>::Create(&OutputString);
+		FJsonSerializer::Serialize(JsonObject.ToSharedRef(), Writer);
+
+		FString ResourceInfoJsonPath = GRDGResourceDumpContext.DumpingDirectoryPath + TEXT("Infos.json");
+		FFileHelper::SaveStringToFile(OutputString, *ResourceInfoJsonPath);
+	}
+
+	// Copy the viewer
+	#if PLATFORM_DESKTOP
+	// TODO: Have the viewer available in cooked build.
+	if (!FPlatformProperties::RequiresCookedData())
+	{
+		const TCHAR* OpenGPUDumpViewerWindowsName = TEXT("OpenGPUDumpViewer.bat");
+		const TCHAR* ViewerHTML = TEXT("GPUDumpViewer.html");
+		FString DumpGPUViewerSourcePath = FPaths::EngineDir() + FString(TEXT("Extras")) / TEXT("GPUDumpViewer");
+
+		PlatformFile.CopyFile(*(GRDGResourceDumpContext.DumpingDirectoryPath / ViewerHTML), *(DumpGPUViewerSourcePath / ViewerHTML));
+		PlatformFile.CopyFile(*(GRDGResourceDumpContext.DumpingDirectoryPath / OpenGPUDumpViewerWindowsName), *(DumpGPUViewerSourcePath / OpenGPUDumpViewerWindowsName));
+	}
+	#endif
+}
+
+void FRDGBuilder::EndResourceDump()
+{
+	if (!GRDGResourceDumpContext.IsDumpingFrame())
+	{
+		return;
+	}
+
+	FString AbsDumpingDirectoryPath = IFileManager::Get().ConvertToAbsolutePathForExternalAppForRead(*GRDGResourceDumpContext.DumpingDirectoryPath);
+	UE_LOG(LogRendererCore, Display, TEXT("Dumped %d resources to %s"), GRDGResourceDumpContext.ResourcesDumped, *AbsDumpingDirectoryPath);
+
+	if (GRDGResourceDumpContext.bShowInExplore)
+	#if PLATFORM_DESKTOP
+	{
+		FPlatformProcess::ExploreFolder(*AbsDumpingDirectoryPath);
+	}
+	#else
+	{
+		UE_LOG(LogRendererCore, Warning, TEXT("r.DumpGPU parameter -show is not supported on this platform."));
+	}
+	#endif
+
+	GRDGResourceDumpContext = FRDGResourceDumpContext();
+}
+
+void FRDGBuilder::DumpResourcePassOutputs(const FRDGPass* Pass)
+{
+	if (bInDebugPassScope)
+	{
+		return;
+	}
+
+	if (!GRDGResourceDumpContext.IsDumpingFrame())
+	{
+		return;
+	}
+
+	// Look whether the pass matches matches r.DumpGPU.Root
+	{
+		FString RootWildcardString = GDumpGPURootCVar.GetValueOnRenderThread();
+		FWildcardString WildcardFilter(RootWildcardString);
+
+		bool bDumpPass = (RootWildcardString == TEXT("*"));
+
+		if (!bDumpPass)
+		{
+			bDumpPass = WildcardFilter.IsMatch(Pass->GetEventName().GetTCHAR());
+		}
+		
+		#if RDG_GPU_SCOPES
+		if (!bDumpPass)
+		{
+			const FRDGEventScope* ParentScope = Pass->GetGPUScopes().Event;
+
+			while (ParentScope)
+			{
+				bDumpPass = bDumpPass || WildcardFilter.IsMatch(ParentScope->Name.GetTCHAR());
+				ParentScope = ParentScope->ParentScope;
+			}
+		}
+		#endif
+
+		if (!bDumpPass)
+		{
+			return;
+		}
+	}
+
+	bInDebugPassScope = true;
+
+	TArray<TSharedPtr<FJsonValue>> InputResourceNames;
+	TArray<TSharedPtr<FJsonValue>> OutputResourceNames;
+	Pass->GetParameters().Enumerate([&](FRDGParameter Parameter)
+	{
+		switch (Parameter.GetType())
+		{
+		case UBMT_RDG_TEXTURE:
+		{
+			if (FRDGTextureRef Texture = Parameter.GetAsTexture())
+			{
+				GRDGResourceDumpContext.AddDumpTexturePass(*this, InputResourceNames, OutputResourceNames, Pass, Texture, /* bIsOutputResource = */ false);
+			}
+		}
+		break;
+		case UBMT_RDG_TEXTURE_SRV:
+		{
+			if (FRDGTextureSRVRef SRV = Parameter.GetAsTextureSRV())
+			{
+				FRDGTextureRef Texture = SRV->Desc.Texture;
+				GRDGResourceDumpContext.AddDumpTexturePass(*this, InputResourceNames, OutputResourceNames, Pass, Texture, /* bIsOutputResource = */ false);
+			}
+		}
+		break;
+		case UBMT_RDG_TEXTURE_UAV:
+		{
+			if (FRDGTextureUAVRef UAV = Parameter.GetAsTextureUAV())
+			{
+				FRDGTextureRef Texture = UAV->Desc.Texture;
+				GRDGResourceDumpContext.AddDumpTexturePass(*this, InputResourceNames, OutputResourceNames, Pass, Texture, /* bIsOutputResource = */ true);
+			}
+		}
+		break;
+		case UBMT_RDG_TEXTURE_ACCESS:
+		{
+			if (FRDGTextureAccess TextureAccess = Parameter.GetAsTextureAccess())
+			{
+				bool bIsOutputResource = (
+					TextureAccess.GetAccess() == ERHIAccess::UAVCompute ||
+					TextureAccess.GetAccess() == ERHIAccess::UAVGraphics ||
+					TextureAccess.GetAccess() == ERHIAccess::RTV);
+
+				GRDGResourceDumpContext.AddDumpTexturePass(*this, InputResourceNames, OutputResourceNames, Pass, TextureAccess, bIsOutputResource);
+			}
+		}
+		break;
+		case UBMT_RDG_TEXTURE_ACCESS_ARRAY:
+		{
+			const FRDGTextureAccessArray& TextureAccessArray = Parameter.GetAsTextureAccessArray();
+
+			for (FRDGTextureAccess TextureAccess : TextureAccessArray)
+			{
+				bool bIsOutputResource = (
+					TextureAccess.GetAccess() == ERHIAccess::UAVCompute ||
+					TextureAccess.GetAccess() == ERHIAccess::UAVGraphics ||
+					TextureAccess.GetAccess() == ERHIAccess::RTV);
+
+				GRDGResourceDumpContext.AddDumpTexturePass(*this, InputResourceNames, OutputResourceNames, Pass, TextureAccess, bIsOutputResource);
+			}
+		}
+		break;
+
+		case UBMT_RDG_BUFFER_SRV:
+		{
+			if (FRDGBufferSRVRef SRV = Parameter.GetAsBufferSRV())
+			{
+				FRDGBufferRef Buffer = SRV->Desc.Buffer;
+				GRDGResourceDumpContext.AddDumpBufferPass(*this, InputResourceNames, OutputResourceNames, Pass, Buffer, /* bIsOutputResource = */ false);
+			}
+		}
+		break;
+		case UBMT_RDG_BUFFER_UAV:
+		{
+			if (FRDGBufferUAVRef UAV = Parameter.GetAsBufferUAV())
+			{
+				FRDGBufferRef Buffer = UAV->Desc.Buffer;
+				GRDGResourceDumpContext.AddDumpBufferPass(*this, InputResourceNames, OutputResourceNames, Pass, Buffer, /* bIsOutputResource = */ true);
+			}
+		}
+		break;
+		case UBMT_RDG_BUFFER_ACCESS:
+		{
+			if (FRDGBufferAccess BufferAccess = Parameter.GetAsBufferAccess())
+			{
+				bool bIsOutputResource = (
+					BufferAccess.GetAccess() == ERHIAccess::UAVCompute ||
+					BufferAccess.GetAccess() == ERHIAccess::UAVGraphics);
+
+				GRDGResourceDumpContext.AddDumpBufferPass(*this, InputResourceNames, OutputResourceNames, Pass, BufferAccess, bIsOutputResource);
+			}
+		}
+		break;
+		case UBMT_RDG_BUFFER_ACCESS_ARRAY:
+		{
+			const FRDGBufferAccessArray& BufferAccessArray = Parameter.GetAsBufferAccessArray();
+
+			for (FRDGBufferAccess BufferAccess : BufferAccessArray)
+			{
+				bool bIsOutputResource = (
+					BufferAccess.GetAccess() == ERHIAccess::UAVCompute ||
+					BufferAccess.GetAccess() == ERHIAccess::UAVGraphics);
+
+				GRDGResourceDumpContext.AddDumpBufferPass(*this, InputResourceNames, OutputResourceNames, Pass, BufferAccess, bIsOutputResource);
+			}
+		}
+		break;
+
+		case UBMT_RENDER_TARGET_BINDING_SLOTS:
+		{
+			const FRenderTargetBindingSlots& RenderTargets = Parameter.GetAsRenderTargetBindingSlots();
+
+			RenderTargets.Enumerate([&](FRenderTargetBinding RenderTarget)
+			{
+				FRDGTextureRef Texture = RenderTarget.GetTexture();
+				GRDGResourceDumpContext.AddDumpTexturePass(*this, InputResourceNames, OutputResourceNames, Pass, Texture, /* bIsOutputResource = */ true);
+			});
+
+			const FDepthStencilBinding& DepthStencil = RenderTargets.DepthStencil;
+
+			if (FRDGTextureRef Texture = DepthStencil.GetTexture())
+			{
+				const bool bHasStoreAction = DepthStencil.GetDepthStencilAccess().IsAnyWrite();
+				GRDGResourceDumpContext.AddDumpTexturePass(*this, InputResourceNames, OutputResourceNames, Pass, Texture, /* bIsOutputResource = */ bHasStoreAction);
+			}
+		}
+		break;
+		}
+	});
+
+	{
+		TArray<TSharedPtr<FJsonValue>> ParentEventScopeNames;
+		#if RDG_GPU_SCOPES
+		{
+			const FRDGEventScope* ParentScope = Pass->GetGPUScopes().Event;
+
+			while (ParentScope)
+			{
+				ParentEventScopeNames.Add(MakeShareable(new FJsonValueString(ParentScope->Name.GetTCHAR())));
+				ParentScope = ParentScope->ParentScope;
+			}
+		}
+		#endif
+
+		TSharedPtr<FJsonObject> JsonObject = MakeShareable(new FJsonObject);
+		JsonObject->SetStringField(TEXT("EventName"), Pass->GetEventName().GetTCHAR());
+		JsonObject->SetStringField(TEXT("ParametersName"), Pass->GetParameters().GetLayout().GetDebugName());
+		JsonObject->SetStringField(TEXT("Pointer"), FString::Printf(TEXT("%p"), Pass));
+		JsonObject->SetNumberField(TEXT("Id"), GRDGResourceDumpContext.PassesCount);
+		JsonObject->SetArrayField(TEXT("ParentEventScopes"), ParentEventScopeNames);
+		JsonObject->SetArrayField(TEXT("InputResources"), InputResourceNames);
+		JsonObject->SetArrayField(TEXT("OutputResources"), OutputResourceNames);
+
+		FString OutputString;
+		TSharedRef<TJsonWriter<TCHAR, TPrettyJsonPrintPolicy<TCHAR>>> Writer = TJsonWriterFactory<TCHAR, TPrettyJsonPrintPolicy<TCHAR>>::Create(&OutputString);
+		FJsonSerializer::Serialize(JsonObject.ToSharedRef(), Writer);
+
+		FString JsonPath = GRDGResourceDumpContext.DumpingDirectoryPath + TEXT("Passes.json");
+		FFileHelper::SaveStringToFile(
+			OutputString, *JsonPath,
+			FFileHelper::EEncodingOptions::AutoDetect, &IFileManager::Get(), FILEWRITE_Append);
+	}
+
+	GRDGResourceDumpContext.PassesCount++;
+
+	bInDebugPassScope = false;
+}
+
+#endif //! RDG_DUMP_RESOURCES
\ No newline at end of file
diff --git a/Engine/Source/Runtime/RenderCore/Public/RenderGraphBuilder.h b/Engine/Source/Runtime/RenderCore/Public/RenderGraphBuilder.h
--- a/Engine/Source/Runtime/RenderCore/Public/RenderGraphBuilder.h
+++ b/Engine/Source/Runtime/RenderCore/Public/RenderGraphBuilder.h
@@ -299,6 +299,11 @@ public:
 	/** The blackboard used to hold common data tied to the graph lifetime. */
 	FRDGBlackboard Blackboard;
 
+#if RDG_DUMP_RESOURCES
+	static void BeginResourceDump(const TArray<FString>& Args);
+	static void EndResourceDump();
+#endif
+
 	//////////////////////////////////////////////////////////////////////////
 	// Deprecated Functions
 	UE_DEPRECATED(5.0, "PreallocateTexture has been renamed to ConvertToExternalTexture")
@@ -576,10 +581,10 @@ private:
 	FRDGUserValidation UserValidation;
 	FRDGBarrierValidation BarrierValidation;
 	FRDGLogFile LogFile;
+#endif
 
 	/** Tracks whether we are in a scope of adding passes to the builder. Used to avoid recursion. */
 	bool bInDebugPassScope = false;
-#endif
 
 #if WITH_MGPU
 	/** Name for the temporal effect used to synchronize multi-frame resources. */
@@ -703,6 +708,10 @@ private:
 		BitArray.Insert(bValue, BitArray.Num(), Passes.Num() - BitArray.Num());
 	}
 
+#if RDG_DUMP_RESOURCES
+	void DumpResourcePassOutputs(const FRDGPass* Pass);
+#endif
+
 #if RDG_ENABLE_DEBUG
 	void VisualizePassOutputs(const FRDGPass* Pass);
 	void ClobberPassOutputs(const FRDGPass* Pass);
diff --git a/Engine/Source/Runtime/RenderCore/Public/RenderGraphDefinitions.h b/Engine/Source/Runtime/RenderCore/Public/RenderGraphDefinitions.h
--- a/Engine/Source/Runtime/RenderCore/Public/RenderGraphDefinitions.h
+++ b/Engine/Source/Runtime/RenderCore/Public/RenderGraphDefinitions.h
@@ -71,6 +71,9 @@
 	#define IF_RDG_CMDLIST_STATS(Op)
 #endif
 
+/** Allows to dump all RDG resources of a frame. */
+#define RDG_DUMP_RESOURCES (!UE_BUILD_SHIPPING)
+
 /** ENUMS */
 
 enum class ERDGBuilderFlags
diff --git a/Engine/Source/Runtime/RenderCore/Public/RenderGraphResources.h b/Engine/Source/Runtime/RenderCore/Public/RenderGraphResources.h
--- a/Engine/Source/Runtime/RenderCore/Public/RenderGraphResources.h
+++ b/Engine/Source/Runtime/RenderCore/Public/RenderGraphResources.h
@@ -842,6 +842,20 @@ struct FRDGBufferDesc
 	EUnderlyingType UnderlyingType = EUnderlyingType::VertexBuffer;
 };
 
+inline const TCHAR* GetBufferUnderlyingTypeName(FRDGBufferDesc::EUnderlyingType BufferType)
+{
+	switch (BufferType)
+	{
+	case FRDGBufferDesc::EUnderlyingType::VertexBuffer:
+		return TEXT("VertexBuffer");
+	case FRDGBufferDesc::EUnderlyingType::StructuredBuffer:
+		return TEXT("StructuredBuffer");
+	case FRDGBufferDesc::EUnderlyingType::AccelerationStructure:
+		return TEXT("AccelerationStructure");
+	}
+	return TEXT("");
+}
+
 struct FRDGBufferSRVDesc final
 	: public FRHIBufferSRVCreateInfo
 {
diff --git a/Engine/Source/Runtime/RenderCore/RenderCore.Build.cs b/Engine/Source/Runtime/RenderCore/RenderCore.Build.cs
--- a/Engine/Source/Runtime/RenderCore/RenderCore.Build.cs
+++ b/Engine/Source/Runtime/RenderCore/RenderCore.Build.cs
@@ -12,11 +12,14 @@ public class RenderCore : ModuleRules
 
 		PrivateIncludePathModuleNames.AddRange(new string[] { "TargetPlatform" });
 
+		// JSON is used for the asset info in the shader library and dumping out frames.
+		PrivateDependencyModuleNames.Add("Json");
+
+		PrivateDependencyModuleNames.Add("BuildSettings");
+
 		if (Target.bBuildEditor == true)
         {
 			DynamicallyLoadedModuleNames.Add("TargetPlatform");
-			// JSON is used for the asset info in the shader library
-			PrivateDependencyModuleNames.Add("Json");
 			// UObjects are used to produce the full path of the asset by which the shaders are identified
 			PrivateDependencyModuleNames.Add("CoreUObject");
 		}
