Improves DumpGPU command

Dumps improvements:
1) Bring up for consoles
2) Better out of memory resiliency during the dumping process
3) Dumps console variables in CSV
4) Dumps process' log after completion
5) Dumps mip chains through FDumpTextureCS compute shader
6) Dumps depth & stencil texture formats through the FDumpTextureCS compute shader
7) Dumps at draw granularity with FRDGBuilder::DumpDraw(); (experimental)
8) Dumps final png screenshot to the dump directory
9) Adds & Dumps the FRDGBufferDesc::Metadata for viewer to decode buffer binary automatically
10) Dumps the PassParameters with structure metadata to decode shader parameters automatically
11) Adds CTRL+SHIFT+/ shortcut

Viewer improvements:
1) Tips display onload to spread some knowledge to the user
2) Supports for opening any pass/resources in new web browser tab
3) Emulates 16 and 32 bits UINT texture visualization with multiple webgl 8bit UINT textures
4) Fixes the webpage's tab going out of memory after visualizing many large resources.
5) Fixes the webpage's tab going out of memory after loading large buffer.
6) Adds support for more texture format with RGB channel reswizzling
7) Implements UI color-sheme based on UE5's editor theme
8) Implements texel color picker capabable of decoding every pixel format.
9) Implements texture viewer zooming with the mouse wheel
10) Implements a r.DumpGPU.Viewer.Visualize to open a specific RDG output resource when opening the viewer

diff --git a/Engine/Config/BaseInput.ini b/Engine/Config/BaseInput.ini
--- a/Engine/Config/BaseInput.ini
+++ b/Engine/Config/BaseInput.ini
@@ -41,6 +41,7 @@ ButtonRepeatDelay=0.1
 +DebugExecBindings=(Key=F9,Command="shot showui")
 +DebugExecBindings=(Key=Period,Command="RECOMPILESHADERS CHANGED",Control=True,Shift=True)
 +DebugExecBindings=(Key=Comma,Command="PROFILEGPU",Control=True,Shift=True)
++DebugExecBindings=(Key=Slash,Command="DumpGPU",Control=True,Shift=True)
 +DebugExecBindings=(Key=Tab,Command="FocusNextPIEWindow",Control=True)
 +DebugExecBindings=(Key=Tab,Command="FocusLastPIEWindow",Control=True,Shift=True)
 +DebugExecBindings=(Key=PageDown,Command="PreviousDebugTarget")
@@ -234,6 +235,7 @@ MaxScrollbackSize=1024
 +ManualAutoCompleteList=(Command="FreezeRendering",Desc="Toggle freezing of most aspects of rendering (such as visibility calculations), useful in conjunction with ToggleDebugCamera to fly around and see how frustum and occlusion culling is working")
 +ManualAutoCompleteList=(Command="ProfileGPU",Desc="Profile one frame of rendering commands sent to the GPU")
 +ManualAutoCompleteList=(Command="ProfileGPUHitches",Desc="Toggle profiling of GPU hitches.")
++ManualAutoCompleteList=(Command="DumpGPU",Desc="Dump one frame of rendering intermediary resources to disk.")
 +ManualAutoCompleteList=(Command="Automation",Desc="Run an automation command (e.g., Automation RunTests TestName)")
 +ManualAutoCompleteList=(Command="CsvProfile Start",Desc="Start CSV profiling.")
 +ManualAutoCompleteList=(Command="CsvProfile Stop",Desc="Stop CSV profiling.")
diff --git a/Engine/Extras/GPUDumpViewer/GPUDumpViewer.html b/Engine/Extras/GPUDumpViewer/GPUDumpViewer.html
--- a/Engine/Extras/GPUDumpViewer/GPUDumpViewer.html
+++ b/Engine/Extras/GPUDumpViewer/GPUDumpViewer.html
@@ -1,23 +1,99 @@
+<!-- Copyright Epic Games, Inc. All Rights Reserved. -->
 <html>
 <title>GPU Dump Viewer</title>
 
 <script type="text/javascript">
 
-//var current_dir="file:///E:/frosty/QAGame/Saved/GPUDumps/WindowsEditor/2021.10.04-10.39.34/";
-var current_dir="./";
+
+// -------------------------------------------------------------------- CONSTANTS
+
+const k_current_dir = "./";
+
+const k_scroll_width = 8;
+const k_null_json_ptr = "0000000000000000";
+
+
+// -------------------------------------------------------------------- GLOBALS
 
 var g_infos = {};
+var g_dump_cvars = {};
 var g_passes = [];
 
+var g_view = null;
+
+
+// -------------------------------------------------------------------- GLOBALS
+
+class IView
+{
+	constructor()
+	{
+
+	}
+
+	setup_html(parent)
+	{
+		parent.innerHTML = '';
+	}
+
+	resize(ctx)
+	{
+
+	}
+
+	get navigations()
+	{
+		return [];
+	}
+
+	release()
+	{
+
+	}
+}
+
+function set_main_view(new_view)
+{
+	var parent_dom = document.getElementById('main_right_pannel');
+	if (g_view !== null)
+	{
+		g_view.release();
+		delete g_view;
+		parent_dom.innerHTML = '';
+	}
+	g_view = new_view;
+	if (g_view !== null)
+	{
+		g_view.setup_html(parent_dom);
+		onresize_body();
+	}
+}
+
 
 // -------------------------------------------------------------------- FILE LOADING
 
+function does_file_exists(relative_path)
+{
+	try
+	{
+		var request = new XMLHttpRequest();
+		request.open('HEAD', k_current_dir + relative_path, false);
+		request.send(null);
+		return request.status != 404;
+	}
+	catch (error)
+	{
+		return false;
+	}
+	return false;
+}
+
 function load_text_file(relative_path)
 {
 	try
 	{
 		var request = new XMLHttpRequest();
-		request.open('GET', current_dir + relative_path, false);
+		request.open('GET', k_current_dir + relative_path, false);
 		request.send(null);
 		if (request.status === 0 || request.status === 200)
 		{
@@ -36,7 +112,7 @@ function load_binary_file(relative_path, callback)
 	try
 	{
 		var request = new XMLHttpRequest();
-		request.open('GET', current_dir + relative_path, true);
+		request.open('GET', k_current_dir + relative_path, true);
 		request.responseType = "arraybuffer";
 
 		request.onload = function(event)
@@ -100,10 +176,254 @@ function load_resource_desc(resource_name)
 	return load_json(`Resources/${resource_name}.json`);
 }
 
+function load_structure_metadata(structure_ptr)
+{
+	var cache = {};
 
-// -------------------------------------------------------------------- DISPLAY PASS
+	function load_nested_structure_metadata(nested_structure_ptr)
+	{
+		if (nested_structure_ptr in cache)
+		{
+			return cache[nested_structure_ptr];
+		}
+
+		var metadata = load_json(`StructuresMetadata/${nested_structure_ptr}.json`);
+		cache[nested_structure_ptr] = metadata;
+ 
+		for (var member of metadata['Members'])
+		{
+			if (member['StructMetadata'] == k_null_json_ptr)
+			{
+				member['StructMetadata'] = null;
+			}
+			else
+			{
+				member['StructMetadata'] = load_nested_structure_metadata(member['StructMetadata']);
+			}
+		}
+
+		return metadata;
+	}
+
+	return load_nested_structure_metadata(structure_ptr);
+}
+
+
+// -------------------------------------------------------------------- UTILITY
+
+function get_filename(file_path)
+{
+	return file_path.split(/(\\|\/)/g).pop();
+}
+
+function px_string_to_int(str)
+{
+	if (Number.isInteger(str))
+	{
+		return str;
+	}
+	if (str.endsWith('px'))
+	{
+		return Number(str.substring(0, str.length - 2));
+	}
+	return Number(str);
+}
+
+function parse_subresource_unique_name(subresource_unique_name)
+{
+	var splitted_name = subresource_unique_name.split(".");
+
+	var subresource_info = {};
+	if (splitted_name[splitted_name.length - 1 ].startsWith('mip') || splitted_name[splitted_name.length - 1 ] == 'stencil')
+	{
+		subresource_info['subresource'] = splitted_name.pop();
+		subresource_info['resource'] = splitted_name.join('.');
+	}
+	else
+	{
+		subresource_info['subresource'] = null;
+		subresource_info['resource'] = subresource_unique_name;
+	}
+	return subresource_info;
+}
+
+function get_subresource_unique_name(subresource_info)
+{
+	if (subresource_info['subresource'] === null)
+	{
+		return subresource_info['resource'];
+	}
+	return `${subresource_info['resource']}.${subresource_info['subresource']}`;
+}
+
+function parse_subresource_unique_version_name(subresource_unique_version_name)
+{
+	var splitted_name = subresource_unique_version_name.split(".");
+
+	var pass_ptr = -1;
+	var draw_id = -1;
+
+	var last = splitted_name.pop();
+	if (last.startsWith('d'))
+	{
+		draw_id = parseInt(last.substring(1));
+		last = splitted_name.pop();
+	}
+	
+	if (last.startsWith('v'))
+	{
+		pass_ptr = last.substring(1);
+	}
+
+	var subresource_unique_name = splitted_name.join('.');
+	var subresource_version_info = parse_subresource_unique_name(subresource_unique_name);
+
+	subresource_version_info['pass'] = pass_ptr;
+	subresource_version_info['draw'] = draw_id;
+	return subresource_version_info;
+}
+
+function get_subresource_unique_version_name(subresource_version_info)
+{
+	if (subresource_version_info['draw'] >= 0)
+	{
+		return `${get_subresource_unique_name(subresource_version_info)}.v${subresource_version_info['pass']}.d${subresource_version_info['draw']}`;
+	}
+	return `${get_subresource_unique_name(subresource_version_info)}.v${subresource_version_info['pass']}`;
+}
+
+
+// -------------------------------------------------------------------- SHADER PARAMETERS
+
+function iterate_structure_members(root_structure_metadata, callback)
+{
+	function iterate_recursive(structure_metadata, offset, cpp_prefix, shader_prefix)
+	{
+		for (var member of structure_metadata['Members'])
+		{
+			var base_type = member['BaseType'];
+
+			if (base_type == 'UBMT_NESTED_STRUCT')
+			{
+				iterate_recursive(member['StructMetadata'], offset + member['Offset'], cpp_prefix + member['Name'] + '.', shader_prefix + member['Name'] + '_');
+			}
+			else if (base_type == 'UBMT_INCLUDED_STRUCT')
+			{
+				iterate_recursive(member['StructMetadata'], offset + member['Offset'], cpp_prefix + member['Name'] + '.', shader_prefix);
+			}
+			else
+			{
+				var params = {
+					member: member,
+					base_type: base_type,
+					byte_offset: offset + member['Offset'],
+					cpp_name: cpp_prefix + member['Name'],
+					shader_name: shader_prefix + member['Name'],
+				};
+				callback(params);
+			}
+		}
+	}
+
+	iterate_recursive(root_structure_metadata, /* offset = */ 0, /* cpp_prefix = */ '', /* shader_prefix = */ '');
+}
+
+
+// -------------------------------------------------------------------- FLOAT ENCODING
+
+function decode_float(raw, total_bit_count, exp_bit_count, has_sign)
+{
+	var exp_bias = (1 << (exp_bit_count - 1)) - 1;
+	var mantissa_bit_count = total_bit_count - exp_bit_count - (has_sign ? 1 : 0);
+
+	var sign_bit      = (raw >> (total_bit_count - 1)) & 0x1;
+	var mantissa_bits = (raw >> 0)                     & ((0x1 << mantissa_bit_count) - 1);
+	var exp_bits      = (raw >> mantissa_bit_count)    & ((0x1 << exp_bit_count) - 1);
 
-var g_pass_id = null;
+	var is_max_exp = exp_bits == ((0x1 << exp_bit_count) - 1);
+
+	var is_denormal = exp_bits == 0;
+	var is_infinity = is_max_exp && mantissa_bits == 0;
+	var is_nan = is_max_exp && mantissa_bits != 0;
+
+	var exp = exp_bits - exp_bias;
+	var mantissa = mantissa_bits * Math.pow(0.5, mantissa_bit_count);
+	var sign = (has_sign && (sign_bit == 1)) ? -1 : 1;
+
+	if (is_nan)
+	{
+		return 'nan';
+	}
+	else if (is_infinity)
+	{
+		return sign == -1 ? '-inf' : '+inf';
+	}
+	else if (is_denormal)
+	{
+		var value = sign * mantissa * Math.pow(0.5, exp_bias - 1);
+		return value;
+	}
+	else
+	{
+		var value = sign * (1.0 + mantissa) * Math.pow(2.0, exp);
+		return value;
+	}
+}
+
+function decode_float10(raw)
+{
+	return decode_float(raw, /* total_bit_count = */ 10, /* exp_bit_count = */ 5, /* has_sign = */ false);
+}
+
+function decode_float11(raw)
+{
+	return decode_float(raw, /* total_bit_count = */ 11, /* exp_bit_count = */ 5, /* has_sign = */ false);
+}
+
+function decode_float16(raw)
+{
+	return decode_float(raw, /* total_bit_count = */ 16, /* exp_bit_count = */ 5, /* has_sign = */ true);
+}
+
+function decode_float32(raw)
+{
+	return decode_float(raw, /* total_bit_count = */ 32, /* exp_bit_count = */ 8, /* has_sign = */ true);
+}
+
+function test_decode_float()
+{
+	var tests = [
+		// Zero
+		[0x0000,  0.0],
+		[0x8000, -0.0],
+
+		// normals
+		[0x4000,  2.0],
+		[0xc000, -2.0],
+
+		// denormals
+		[0x0001,  5.960464477539063e-8],
+		[0x8001, -5.960464477539063e-8],
+
+		// exotics
+		[0x7c00, '+inf'],
+		[0xFc00, '-inf'],
+		[0x7c01, 'nan'],
+		[0xFc01, 'nan'],
+	];
+
+	for (var i = 0; i < tests.length; i++)
+	{
+		var encoded = tests[i][0];
+		var ref = tests[i][1];
+		var computed = decode_float16(encoded);
+
+		console.assert(computed == ref);
+	}
+}
+
+
+// -------------------------------------------------------------------- DISPLAY PASS
 
 function display_pass_hierarchy()
 {
@@ -130,8 +450,13 @@ function display_pass_hierarchy()
 		if (search_resource != '')
 		{
 			show_resource = false;
+			pass_data['InputResources'].forEach(function(output_name) {
+				var subresource_unique_name = parse_subresource_unique_name(output_name)['resource'];
+				show_resource = show_resource || subresource_unique_name.toLowerCase().includes(search_resource.toLowerCase());
+			});
 			pass_data['OutputResources'].forEach(function(output_name) {
-				show_resource = show_resource || output_name.toLowerCase().includes(search_resource.toLowerCase());
+				var subresource_unique_name = parse_subresource_unique_name(output_name)['resource'];
+				show_resource = show_resource || subresource_unique_name.toLowerCase().includes(search_resource.toLowerCase());
 			});
 		}
 
@@ -157,917 +482,3393 @@ function display_pass_hierarchy()
 			for (var i = shared_scope; i < pass_event_scopes.length; i++)
 			{
 				var scope = pass_event_scopes[pass_event_scopes.length - 1 - i];
-				html += `<div style="padding-left: ${10 + 16 * i}px;" class="disabled">${scope}</div>`;
+				html += `<a style="padding-left: ${10 + 16 * i}px;" class="disabled">${scope}</a>`;
 				parent_event_scopes.push(scope);
 			}
 
-			var classes = '';
-			if (pass_id === g_pass_id)
-			{
-				classes = 'selected';
-			}
-
-			html += `<div style="padding-left: ${10 + 16 * pass_event_scopes.length}px;" onmousedown="display_pass(${pass_id});" class="${classes}">${pass_data['EventName']}</div>`;
+			html += `<a style="padding-left: ${10 + 16 * pass_event_scopes.length}px;" href="#display_pass(${pass_id});">${pass_data['EventName']}</a>`;
 		}
 	});
 	document.getElementById('pass_hierarchy').innerHTML = html;
+	update_href_selection(document.getElementById('pass_hierarchy'));
 }
 
-function display_input_resource(pass_id, resource_name)
+class ResourceView extends IView
 {
-	var previous_producer_pass = -1;
-	for (var i = 0; i < pass_id; i++)
+	constructor(subresource_version_info, resource_desc)
 	{
-		var cur_outputs = g_passes[i]['OutputResources'];
-
-		cur_outputs.forEach(function(value) {
-			if (value == resource_name)
-			{
-				previous_producer_pass = i;
-			}
-		});
+		super();
+		this.subresource_version_info = subresource_version_info;
+		this.resource_desc = resource_desc;
 	}
 
-	return display_output_resource(previous_producer_pass, resource_name);
+	get navigations()
+	{
+		return [];
+	}
 }
 
-function display_output_resource(pass_id, resource_name)
+class PassView extends IView
 {
-	var resource_version_name = `${resource_name}.v0000000000000000`;
-	if (pass_id >= 0)
+	constructor(pass_id)
 	{
-		var pass_data = g_passes[pass_id];
-		resource_version_name = `${resource_name}.v${pass_data['Pointer']}`;
-	}
+		super();
+		this.pass_id = pass_id;
+		this.pass_data = g_passes[pass_id];
+		this.pass_draws_data = [];
+		this.resource_view = null;
 
-	var resource_desc = load_resource_desc(resource_name);
-	if (!resource_desc)
-	{
-		return;
+		if (this.pass_data['DrawCount'] > 0)
+		{
+			this.pass_draws_data = load_json_dict_sequence(`Passes/Pass.${this.pass_data['Pointer']}.Draws.json`);
+		}
 	}
 
-	if (resource_desc['Type'] == 'Texture2D')
-	{
-		return display_texture(resource_name, resource_version_name);
-	}
-	else if (resource_desc['Type'] == 'VertexBuffer' || resource_desc['Type'] == 'StructuredBuffer')
+	setup_html(parent_dom)
 	{
-		return display_buffer(resource_name, resource_version_name);
+		var column_width = '50%';
+		var draw_column_display = 'none';
+
+		if (this.pass_draws_data.length > 0)
+		{
+			column_width = '33%';
+			draw_column_display = 'block';
+		}
+
+		parent_dom.innerHTML = `
+			<div class="pass_title main_div">
+				${this.pass_data['EventName']}
+				<div class="button_list" style="display:inline-block;"><a href="#display_pass_parameters(${this.pass_id});">PassParameters</a></div>
+			</div>
+			<table width="100%">
+				<tr>
+					<td width="${column_width}">
+						<div class="main_div">
+							<div class="selection_list_title">Input resources</div>
+							<div class="selection_list_search">
+								<input type="search" id="pass_input_resource_search" oninput="g_view.refresh_input_resource_list();" onchange="g_view.refresh_input_resource_list();" placeholder="Search input resource..." />
+							</div>
+							<div class="selection_list" id="pass_input_resource_list"></div>
+						</div>
+					</td>
+					<td width="${column_width}">
+						<div class="main_div">
+							<div class="selection_list_title">Output resources</div>
+							<div class="selection_list_search">
+								<input type="search" id="pass_output_resource_search" oninput="g_view.refresh_output_resource_list();" onchange="g_view.refresh_output_resource_list();" placeholder="Search output resource..." />
+							</div>
+							<div class="selection_list" id="pass_output_resource_list"></div>
+						</div>
+					</td>
+					<td width="${column_width}">
+						<div class="main_div" style="display: ${draw_column_display};">
+							<div class="selection_list_title">Draws</div>
+							<div class="selection_list_search">
+								<input type="search" id="pass_draw_search" oninput="g_view.refresh_draw_resource_list();" onchange="g_view.refresh_draw_resource_list();" placeholder="Search draw..." />
+							</div>
+							<div class="selection_list" id="pass_draw_list"></div>
+						</div>
+					</td>
+				</tr>
+			</table>
+			<div id="display_resource_pannel"></div>
+			<table width="100%">
+				<tr>
+					<td width="50%">
+						<div class="main_div" id="resource_pass_modifying_outter">
+							<div class="selection_list_title" id="resource_pass_modifying_title"></div>
+							<div class="selection_list_search">
+								<input type="search" id="resource_pass_modifying_search" oninput="g_view.refresh_resource_modifying_list();" onchange="g_view.refresh_resource_modifying_list();" placeholder="Search modifying pass..." />
+							</div>
+							<div class="selection_list" id="resource_pass_modifying_list"></div>
+						</div>
+					</td>
+					<td width="50%">
+						<div class="main_div" id="resource_pass_reading_outter">
+							<div class="selection_list_title" id="resource_pass_reading_title"></div>
+							<div class="selection_list_search">
+								<input type="search" id="resource_pass_reading_search" oninput="g_view.refresh_resource_reading_list();" onchange="g_view.refresh_resource_reading_list();" placeholder="Search reading pass..." />
+							</div>
+							<div class="selection_list" id="resource_pass_reading_list"></div>
+						</div>
+					</td>
+				</tr>
+			</table>
+			`;
 	}
-}
 
-function display_pass(pass_id)
-{
-	if (pass_id == g_pass_id)
+	refresh_all_lists()
 	{
-		return;
+		this.refresh_input_resource_list();
+		this.refresh_output_resource_list();
+		this.refresh_draw_resource_list();
+		this.refresh_resource_modifying_list();
+		this.refresh_resource_reading_list();
 	}
 
-	var pass_data = g_passes[pass_id];
-	document.getElementById('main_right_pannel').innerHTML = `<div class="pass_title main_div">${pass_data['EventName']}</div>`;
-
-	// Display inputs
+	refresh_input_resource_list()
 	{
-		var html = `<div class="main_div">
-			<div class="selection_list_title">Input resources: ${pass_data['EventName']}</div>
-			<div class="selection_list">`;
-		pass_data['InputResources'].forEach(function(resource_name) {
-			var resource_desc = load_resource_desc(resource_name);
+		var display_list = [];
+		for (const subresource_unique_name of this.pass_data['InputResources'])
+		{
+			var subresource_info = parse_subresource_unique_name(subresource_unique_name);
+			var resource_desc = load_resource_desc(subresource_info['resource']);
+
+			var name = subresource_unique_name;
+			var href = null;
 			if (resource_desc)
 			{
-				html += `<div onmousedown="display_input_resource(${pass_id}, '${resource_name}');">${resource_desc['Name']}</div>`;
-			}
-			else
-			{
-				html += `<div>${resource_name}</div>`;
+				if (subresource_info['subresource'] == null)
+				{
+					name = resource_desc['Name'];
+				}
+				else
+				{
+					name = `${resource_desc['Name']} ${subresource_info['subresource']}`;
+				}
+				href = `#display_input_resource(${this.pass_id},'${subresource_unique_name}');`
 			}
-		});
-		html += '</div></div>';
-		document.getElementById('main_right_pannel').innerHTML += html;
-	}
-
-	// Display outputs
-	{
-		var html = `<div class="main_div">
-			<div class="selection_list_title">Output resources: ${pass_data['EventName']}</div>
-			<div class="selection_list">`;
-		pass_data['OutputResources'].forEach(function(resource_name) {
-			var resource_desc = load_resource_desc(resource_name);
-			html += `<div onmousedown="display_output_resource(${pass_id}, '${resource_name}');">${resource_desc['Name']}</div>`;
-		});
-		html += '</div></div>';
-		document.getElementById('main_right_pannel').innerHTML += html;
-	}
 
-	document.getElementById('main_right_pannel').innerHTML += `<div id="display_resource_pannel"></div>`;
+			display_list.push({'name': name, 'href': href});
+		}
 
-	g_pass_id = pass_id;
+		render_selection_list_html(
+			document.getElementById('pass_input_resource_list'),
+			display_list,
+			{
+				'search': document.getElementById('pass_input_resource_search').value,
+				'deduplicate': true,
+				'sort': true
+			});
+	}
 
-	// Load first resource
+	refresh_output_resource_list()
 	{
-		var resource_name = pass_data['OutputResources'][0];
-
-		if (resource_name)
+		var draw_id = -1;
+		if (this.resource_view)
 		{
-			display_output_resource(pass_id, resource_name);
+			draw_id = this.resource_view.subresource_version_info['draw'];
 		}
-	}
-
-	display_pass_hierarchy();
-}
 
+		var display_list = [];
+		for (const subresource_unique_name of this.pass_data['OutputResources'])
+		{
+			var subresource_info = parse_subresource_unique_name(subresource_unique_name);
+			var resource_desc = load_resource_desc(subresource_info['resource']);
 
-// -------------------------------------------------------------------- DISPLAY GENERAL RESOURCE INFO
+			var name = subresource_unique_name;
+			var href = null;
+			if (resource_desc)
+			{
+				if (subresource_info['subresource'] == null)
+				{
+					name = resource_desc['Name'];
+				}
+				else
+				{
+					name = `${resource_desc['Name']} ${subresource_info['subresource']}`;
+				}
 
-function get_resource_producer_html(resource_name)
-{
-	var resource_desc = load_resource_desc(resource_name);
+				if (draw_id >= 0)
+				{
+					href = `#display_draw_output_resource(${this.pass_id},'${subresource_unique_name}',${draw_id});`;
+				}
+				else
+				{
+					href = `#display_output_resource(${this.pass_id},'${subresource_unique_name}');`;
+				}
+			}
 
-	var html = `
-		<div class="main_div">
-			<div class="selection_list_title">Passes producing ${resource_desc['Name']}</div>
-			<div class="selection_list">`;
+			display_list.push({'name': name, 'href': href});
+		}
 
-	for (var pass_id = 0; pass_id < g_passes.length; pass_id++)
-	{
-		var producer = false;
-		g_passes[pass_id]['OutputResources'].forEach(function(value) {
-			if (value == resource_name)
+		render_selection_list_html(
+			document.getElementById('pass_output_resource_list'),
+			display_list,
 			{
-				producer = true;
-			}
-		});
+				'search': document.getElementById('pass_output_resource_search').value,
+			});
+	}
 
-		if (producer)
+	refresh_draw_resource_list()
+	{
+		var is_output_resource = false;
+		var subresource_unique_name = '';
+		if (this.resource_view)
 		{
-			html += `<div onmousedown="display_pass(${pass_id}); display_output_resource(${pass_id}, '${resource_name}');">${g_passes[pass_id]['EventName']}</div>`;
+			is_output_resource = this.pass_data['OutputResources'].includes(get_subresource_unique_name(this.resource_view.subresource_version_info));
+			subresource_unique_name = get_subresource_unique_name(this.resource_view.subresource_version_info);
 		}
-	}
-
-	html += `
-			</div>
-		</div>`;
 
-	html += `
-		<div class="main_div">
-			<div class="selection_list_title">Passes reading ${resource_desc['Name']}</div>
-			<div class="selection_list">`;
-
-	for (var pass_id = 0; pass_id < g_passes.length; pass_id++)
-	{
-		var reader = false;
-		g_passes[pass_id]['InputResources'].forEach(function(value) {
-			if (value == resource_name)
+		var display_list = [];
+		for (var draw_id = 0; draw_id < this.pass_draws_data.length; draw_id++)
+		{
+			var href = null;
+			if (subresource_unique_name && is_output_resource)
 			{
-				reader = true;
+				href = `#display_draw_output_resource(${this.pass_id},'${subresource_unique_name}',${draw_id});`;
 			}
-		});
 
-		if (reader)
-		{
-			html += `<div onmousedown="display_pass(${pass_id}); display_input_resource(${pass_id}, '${resource_name}');">${g_passes[pass_id]['EventName']}</div>`;
+			display_list.push({
+				'name': `${draw_id}: ${this.pass_draws_data[draw_id]['DrawName']}`,
+				'href': href
+			});
 		}
-	}
-
-	html += `
-			</div>
-		</div>`;
-
-	return html;
-}
-
 
-// -------------------------------------------------------------------- DISPLAY TEXTURE
-
-var g_texture_name = null;
-var g_texture_version_name = null;
-var g_gl = null;
-var g_gl_ext = null;
-var g_gl_texture = null;
+		if (subresource_unique_name)
+		{
+			var href = null;
+			if (is_output_resource)
+			{
+				href = `#display_output_resource(${this.pass_id},'${subresource_unique_name}');`;
+			}
 
-var g_shader_code_dict = {};
+			display_list.push({
+				'name': 'Final pass output',
+				'href': href
+			});
+		}
 
-function get_default_texture_shader_code(gl, resource_desc, resource_version_name)
-{
-	if (resource_version_name in g_shader_code_dict)
-	{
-		return g_shader_code_dict[resource_version_name];
+		render_selection_list_html(
+			document.getElementById('pass_draw_list'),
+			display_list,
+			{
+				'search': document.getElementById('pass_draw_search').value,
+			});
 	}
 
-	var gl_format = translate_pixel_format(gl, resource_desc['Format']);
-
-	if (!gl_format)
+	refresh_resource_modifying_list()
 	{
-		return `// Sorry but ${resource_desc['Format']} is not supported for display :(`;
-	}
-
-	var shader_program = `// WebGL visualization shader for a ${gl_format['ue_pixel_format']}
-uniform ${gl_format['sampler_type']} tex_sampler;
-
-in vec2 tex_coordinates;
-out vec4 color;
-
-void main(void)
-{
-	ivec2 texel_coord = ivec2(tex_coordinates * vec2(textureSize(tex_sampler, 0)));
-	${gl_format['sampler_return_type']} raw = texelFetch(tex_sampler, texel_coord , 0);
-	color = ${gl_format['swizzling']};
-}`;
-	return shader_program;
-}
-
-function onresize_texture_shader_user_code_change()
-{
-	refresh_texture_visualization(g_texture_name, g_texture_version_name, null);
-}
-
-function display_texture(resource_name, resource_version_name)
-{
-	var resource_desc = load_resource_desc(resource_name);
+		if (this.resource_view === null || this.resource_view.resource_desc['Desc'] == 'FRDGParameterStruct')
+		{
+			document.getElementById('resource_pass_modifying_outter').style.display = 'none';
+			return;
+		}
+		document.getElementById('resource_pass_modifying_outter').style.display = 'block';
 
-	if (resource_version_name == g_texture_version_name)
-	{
-		return;
-	}
+		var subresource_unique_name = get_subresource_unique_name(this.resource_view.subresource_version_info);
+		var resource_desc = load_resource_desc(this.resource_view.subresource_version_info['resource']);
 
-	g_texture_name = resource_name;
-	g_texture_version_name = resource_version_name;
-	g_gl = null;
-	g_gl_ext = null;
-	g_gl_texture = null;
+		var display_list = [];
+		for (var pass_id = 0; pass_id < g_passes.length; pass_id++)
+		{
+			var producer = false;
+			g_passes[pass_id]['OutputResources'].forEach(function(value) {
+				if (value == subresource_unique_name)
+				{
+					producer = true;
+				}
+			});
 
-	var canvas_w = resource_desc['ExtentX'];
-	var canvas_h = resource_desc['ExtentY'];
+			if (producer)
+			{
+				display_list.push({
+					'name': g_passes[pass_id]['EventName'],
+					'href': `#display_output_resource(${pass_id},'${subresource_unique_name}');`,
+				});
+			}
+		}
 
-	if (canvas_w < 512)
-	{
-		canvas_h = 512 * canvas_h / canvas_w;
-		canvas_w = 512;
+		document.getElementById('resource_pass_modifying_title').innerHTML = `Passes modifying ${resource_desc['Name']}`;
+		render_selection_list_html(
+			document.getElementById('resource_pass_modifying_list'),
+			display_list,
+			{
+				'search': document.getElementById('resource_pass_modifying_search').value,
+			});
 	}
 
-	var resource_info_htmls = '';
+	refresh_resource_reading_list()
 	{
-		resource_info_htmls += `
-			<table width="100%" class="resource_desc">`;
-		
-		for (var key in resource_desc){
-		    if (resource_desc.hasOwnProperty(key)){
-				resource_info_htmls += `
-					<tr>
-						<td>${key}</td>
-						<td>${resource_desc[key]}</td>
-					</tr>`;
-		    }
+		if (this.resource_view === null || this.resource_view.resource_desc['Desc'] == 'FRDGParameterStruct')
+		{
+			document.getElementById('resource_pass_reading_outter').style.display = 'none';
+			return;
 		}
+		document.getElementById('resource_pass_reading_outter').style.display = 'block';
 
-		resource_info_htmls += `
-			</table>`;
-	}
+		var subresource_unique_name = get_subresource_unique_name(this.resource_view.subresource_version_info);
+		var resource_desc = load_resource_desc(this.resource_view.subresource_version_info['resource']);
 
-	var html = '';
-	html += `
-		<div class="main_div">
-			<div class="selection_list_title">Texture visualization: ${resource_desc['Name']}</div>
-			<div><canvas id="texture_visualization_canvas" width="${canvas_w}px" height="${canvas_h}px"></canvas></div>
-			<table width="100%">
-				<tr>
-					<td>
-						<div class="selection_list_title">Texture descriptor</div>
-						${resource_info_htmls}
-					</td>
-					<td style="width: 50%;">
-						<textarea id="texture_visualization_code_input" oninput="onresize_texture_shader_user_code_change();" onchange="onresize_texture_shader_user_code_change();"></textarea>
-						<textarea id="texture_visualization_code_log" readonly></textarea>
-					</td>
-				</tr>
-			</table>
-		</div>`;
+		var display_list = [];
+		for (var pass_id = 0; pass_id < g_passes.length; pass_id++)
+		{
+			var reader = false;
+			g_passes[pass_id]['InputResources'].forEach(function(value) {
+				if (value == subresource_unique_name)
+				{
+					reader = true;
+				}
+			});
 
-	document.getElementById('display_resource_pannel').innerHTML = html + get_resource_producer_html(resource_name);
+			if (reader)
+			{
+				display_list.push({
+					'name': g_passes[pass_id]['EventName'],
+					'href': `#display_input_resource(${pass_id},'${subresource_unique_name}');`,
+				});
+			}
+		}
 
-	var canvas = document.getElementById('texture_visualization_canvas');
-	g_gl = canvas.getContext('webgl2');
+		document.getElementById('resource_pass_reading_title').innerHTML = `Passes reading ${resource_desc['Name']}`;
+		render_selection_list_html(
+			document.getElementById('resource_pass_reading_list'),
+			display_list,
+			{
+				'search': document.getElementById('resource_pass_reading_search').value,
+			});
+	}
 
+	resize(ctx)
 	{
-		var available_extensions = g_gl.getSupportedExtensions();
-		var required_extensions = ['EXT_texture_norm16'];
-
-		g_gl_ext = {};
-		required_extensions.forEach(function(ext_name)
+		if (this.resource_view !== null)
 		{
-			g_gl_ext[ext_name] = g_gl.getExtension(ext_name);
-		});
+			this.resource_view.resize(ctx);
+		}
 	}
 
-	document.getElementById('texture_visualization_code_input').value = get_default_texture_shader_code(g_gl, resource_desc, resource_version_name);
-
-	load_binary_file(`Resources/${resource_version_name}.bin`, function(raw_pixel_data)
+	set_resource_view(new_resource_view)
 	{
-		if (raw_pixel_data)
+		var parent_dom = document.getElementById('display_resource_pannel');
+		if (this.resource_view !== null)
 		{
-			refresh_texture_visualization(resource_name, resource_version_name, raw_pixel_data);
+			this.resource_view.release();
+			delete this.resource_view;
+			parent_dom.innerHTML = '';
 		}
-		else
-		{
 
+		this.resource_view = new_resource_view;
+		if (this.resource_view !== null)
+		{
+			this.resource_view.setup_html(parent_dom);
+			this.refresh_all_lists();
+			onresize_body();
 		}
-	});
-}
-
-function translate_pixel_format(gl, ue_pixel_format)
-{
+	}
 
-	// Missing formats:
-	//PF_G16R16F_FILTER
-	//PF_A2B10G10R10
-	//PF_D24
-	//PF_R16F_FILTER
-	//PF_V8U8
-	//PF_R32_UINT
-	//PF_R32_SINT
-	//PF_PVRTC2
-	//PF_PVRTC4
-	//PF_R16_UINT
-	//PF_R16_SINT
-	//PF_R16G16B16A16_UINT
-	//PF_R16G16B16A16_SINT
-	//PF_R5G6B5_UNORM
-	//PF_R32G32B32A32_UINT
-	//PF_R16G16_UINT
-	//PF_BC6H
-	//PF_R8G8B8A8_SNORM
-	//PF_R16G16B16A16_UNORM
-	//PF_R16G16B16A16_SNORN
-	//PF_R32G32_UINT
+	get navigations()
+	{
+		var navs = [`display_pass(${this.pass_id});`];
 
-	var internal_format = null;
-	var src_format = null;
-	var src_type = null;
-	var swizzling = 'raw';
-	var sampler_type = 'sampler2D';
-	var sampler_return_type = 'vec4';
+		if (this.resource_view !== null)
+		{
+			navs.concat(this.resource_view.navigations);
+		}
 
-	// 32bit floats
-	if (ue_pixel_format === 'PF_A32B32G32R32F')
-	{
-		internal_format = gl.RGBA32F;
-		src_format = gl.RGBA;
-		src_type = gl.FLOAT;
-	}
-	else if (ue_pixel_format === 'PF_G32R32F')
-	{
-		internal_format = gl.RG32F;
-		src_format = gl.RG;
-		src_type = gl.FLOAT;
-	}
-	else if (ue_pixel_format === 'PF_R32_FLOAT')
-	{
-		internal_format = gl.R32F;
-		src_format = gl.RED;
-		src_type = gl.FLOAT;
+		return navs;
 	}
-	// 16bit floats
-	else if (ue_pixel_format === 'PF_FloatRGBA')
+
+	release()
 	{
-		internal_format = gl.RGBA16F;
-		src_format = gl.RGBA;
-		src_type = gl.HALF_FLOAT;
+		this.set_resource_view(null);
 	}
-	else if (ue_pixel_format === 'PF_G16R16F')
+}
+
+function display_pass_internal(pass_id)
+{
+	if (g_view instanceof PassView && pass_id == g_view.pass_id)
 	{
-		internal_format = gl.RG16F;
-		src_format = gl.RG;
-		src_type = gl.HALF_FLOAT;
+		return;
 	}
-	else if (ue_pixel_format === 'PF_R16F')
+
+	set_main_view(new PassView(pass_id));
+}
+
+function display_resource_internal(subresource_version_info)
+{
+	var resource_desc = load_resource_desc(subresource_version_info['resource']);
+	if (!resource_desc)
 	{
-		internal_format = gl.R16F;
-		src_format = gl.RED;
-		src_type = gl.HALF_FLOAT;
+		return;
 	}
-	// 11-11-10bit floats
-	else if (ue_pixel_format === 'PF_FloatRGB' || ue_pixel_format === 'PF_FloatR11G11B10')
+
+	if (g_view.resource_view !== null && get_subresource_unique_version_name(g_view.resource_view.subresource_version_info) == get_subresource_unique_version_name(subresource_version_info))
 	{
-		internal_format = gl.R11F_G11F_B10F;
-		src_format = gl.RGB;
-		src_type = gl.UNSIGNED_INT_10F_11F_11F_REV;
+		return;
 	}
-	// 16bit unorms
-	else if (ue_pixel_format === 'PF_A16B16G16R16')
+
+	if (resource_desc['Type'] == 'Texture2D')
 	{
-		internal_format = g_gl_ext['EXT_texture_norm16'].RGBA16_EXT;
-		src_format = gl.RGBA;
-		src_type = gl.UNSIGNED_SHORT;
+		g_view.set_resource_view(new TextureView(subresource_version_info, resource_desc));
 	}
-	else if (ue_pixel_format === 'PF_G16R16')
+	else if (resource_desc['Type'] == 'VertexBuffer' || resource_desc['Type'] == 'StructuredBuffer')
 	{
-		internal_format = g_gl_ext['EXT_texture_norm16'].RG16_EXT;
-		src_format = gl.RG;
-		src_type = gl.UNSIGNED_SHORT;
+		g_view.set_resource_view(new BufferView(subresource_version_info, resource_desc));
 	}
-	else if (ue_pixel_format === 'PF_G16')
+}
+
+
+// -------------------------------------------------------------------- HREF FUNCTIONS
+
+function display_pass(pass_id)
+{
+	// Load first resource
+	if (g_passes[pass_id]['OutputResources'][0])
 	{
-		internal_format = g_gl_ext['EXT_texture_norm16'].R16_EXT;
-		src_format = gl.RED;
-		src_type = gl.UNSIGNED_SHORT;
+		redirect_to_hash(`display_output_resource(${pass_id},'${g_passes[pass_id]['OutputResources'][0]}');`);
 	}
-	// 8bit unorms
-	else if (ue_pixel_format === 'PF_B8G8R8A8' || ue_pixel_format === 'PF_R8G8B8A8')
+	else if (g_passes[pass_id]['InputResources'][0])
 	{
-		internal_format = gl.RGBA;
-		src_format = gl.RGBA;
-		src_type = gl.UNSIGNED_BYTE;
+		redirect_to_hash(`display_input_resource(${pass_id},'${g_passes[pass_id]['InputResources'][0]}');`);
 	}
-	else if (ue_pixel_format === 'PF_R8G8')
+	else
 	{
-		internal_format = gl.RG8;
-		src_format = gl.RG;
-		src_type = gl.UNSIGNED_BYTE;
+		display_pass_internal(pass_id);
 	}
-	else if (ue_pixel_format === 'PF_R8' || ue_pixel_format === 'PF_G8' || ue_pixel_format === 'PF_A8' || ue_pixel_format === 'PF_L8')
+}
+
+function display_input_resource(pass_id, subresource_unique_name)
+{
+	var previous_producer_pass = -1;
+	for (var i = 0; i < pass_id; i++)
 	{
-		internal_format = gl.R8;
-		src_format = gl.RED;
-		src_type = gl.UNSIGNED_BYTE;
+		var cur_outputs = g_passes[i]['OutputResources'];
+
+		cur_outputs.forEach(function(value) {
+			if (value == subresource_unique_name)
+			{
+				previous_producer_pass = i;
+			}
+		});
 	}
-	// 8bit uint
-	else if (ue_pixel_format === 'PF_R8G8B8A8_UINT')
+
+	var subresource_version_info = parse_subresource_unique_name(subresource_unique_name);
+	if (previous_producer_pass >= 0)
 	{
-		internal_format = gl.RGBA8UI;
-		src_format = gl.RGBA_INTEGER;
-		src_type = gl.UNSIGNED_BYTE;
+		subresource_version_info['pass'] = g_passes[previous_producer_pass]['Pointer'];
 	}
-	else if (ue_pixel_format === 'PF_R8_UINT')
+	else
 	{
-		internal_format = gl.R8UI;
-		src_format = gl.RED_INTEGER;
-		src_type = gl.UNSIGNED_BYTE;
+		subresource_version_info['pass'] = k_null_json_ptr;
 	}
-	else
+
+	display_pass_internal(pass_id);
+	display_resource_internal(subresource_version_info);
+}
+
+function display_output_resource(pass_id, subresource_unique_name)
+{
+	var subresource_version_info = parse_subresource_unique_name(subresource_unique_name);
+	subresource_version_info['pass'] = g_passes[pass_id]['Pointer'];
+
+	display_pass_internal(pass_id);
+	display_resource_internal(subresource_version_info);
+}
+
+function display_draw_output_resource(pass_id, subresource_unique_name, draw_id)
+{
+	var subresource_version_info = parse_subresource_unique_name(subresource_unique_name);
+	subresource_version_info['pass'] = g_passes[pass_id]['Pointer'];
+	subresource_version_info['draw'] = draw_id;
+
+	display_pass_internal(pass_id);
+	display_resource_internal(subresource_version_info);
+}
+
+
+// -------------------------------------------------------------------- DISPLAY INFOS
+
+class InfosView extends IView
+{
+	constructor()
 	{
-		return null;
+		super();
 	}
 
-	if (src_format === gl.RED_INTEGER || src_format === gl.RG_INTEGER || src_format === gl.RGBA_INTEGER)
+	setup_html(parent_dom)
 	{
-		var divide = 255;
-		if (src_type == gl.UNSIGNED_SHORT)
+		var info_htmls = '';
+		{
+			info_htmls += `
+				<table width="100%" class="pretty_table">`;
+			
+			for (var key in g_infos)
+			{
+			    if (g_infos.hasOwnProperty(key))
+			    {
+					info_htmls += `
+						<tr>
+							<td width="100px">${key}</td>
+							<td>${g_infos[key]}</td>
+						</tr>`;
+			    }
+			}
+
+			info_htmls += `
+				</table>`;
+		}
+
+		parent_dom.innerHTML = `
+			<div class="main_div">
+				<div class="pass_title">Infos</div>
+				<div id="infos_pannel">${info_htmls}</div>
+			</div>`;
+
+		if (does_file_exists('Base/Screenshot.png'))
 		{
-			divide = 65535;
+			parent_dom.innerHTML += `
+				<div class="main_div">
+					<div class="pass_title">Screenshot</div>
+					<img src="${k_current_dir}/Base/Screenshot.png" style="width: 100%;" />
+				</div>`;
 		}
 
-		sampler_type = 'usampler2D';
-		sampler_return_type = 'uvec4';
-		swizzling = `(vec4(${swizzling}) * (1.0 / ${divide}.0))`;
+		document.title = 'Dump infos';
 	}
 
-	if (src_format === gl.RED)
-	{
-		swizzling = `vec4((${swizzling}).r, 0.0, 0.0, 1.0)`;
-	}
-	else if (src_format === gl.RG)
-	{
-		swizzling = `vec4((${swizzling}).rg, 0.0, 1.0)`;
-	}
-	else if (src_format === gl.RGB)
+	get navigations()
 	{
-		swizzling = `vec4((${swizzling}).rgb, 1.0)`;
+		return [`display_infos();`];
 	}
+}
 
-	var gl_format = {};
-	gl_format['internal_format'] = internal_format;
-	gl_format['src_format'] = src_format;
-	gl_format['src_type'] = src_type;
-	gl_format['swizzling'] = swizzling;
-	gl_format['sampler_type'] = sampler_type;
-	gl_format['sampler_return_type'] = sampler_return_type;
-	gl_format['ue_pixel_format'] = ue_pixel_format;
-	return gl_format;
+function display_infos(tip_id)
+{
+	set_main_view(new InfosView());
 }
 
-function create_texture_from_pixel_data(gl, resource_desc, resource_version_name, raw_pixel_data)
+
+// -------------------------------------------------------------------- DISPLAY TIP
+
+const k_tips = [
+	// Dumping process
+	'Can use the CTRL+SHIFT+/ keyboard shortcut to summon the DumpGPU command.',
+	'Speed up your frame dump by only selecting the passes you need with r.DumpGPU.Root. For instance r.DumpGPU.Root="*PostProcessing*".',
+	'GPU dumps can be large and accumulate on your hard drive in your various projects\' Saved/ directories. Set r.DumpGPU.Directory="D:/tmp/DumpGPU/" in your console variables to dump them all at the same location.',
+	'Uses -cvarsini to override console variables with your own ConsoleVariables.ini file on a cooked build.',
+
+	// General navigations
+	'All navigation links can be open in a new tab.',
+	'Uses the browser\'s back button to navigate to previously inspected resource.',
+	'Make sure to browse the dump informations.',
+	'Make sure to browse the console variables.',
+	'Make sure to browse the log file.',
+	'Share the part of the URL after the # (for instance #display_input_resource(96,\'TSR.AntiAliasing.Noise.000000006b9b2c00.mip0\');) to anyone else whom have this GPU dump to so they too can navigate to the exact same resource view.',
+	'Set r.DumpGPU.Viewer.Visualize in your ConsoleVariables.ini so the dump viewer automatically open this RDG output resource at startup.',
+
+	// Buffer visualization
+	'Uses the templated FRDGBufferDesc::Create*<FMyStructure>(NumElements) to display your buffer more conveniently with FMyStructure layout in the buffer visualization.',
+	'Buffer visualization supports float, half, int, uint, short, ushort, char, uchar, as well as hexadecimal with hex(uint) and binary with bin(uint).',
+	'Uses the "Address..." field at the very left of the buffer view\'s header to navigate to a specific address. Supports decimal and 0x prefixed hexadecimal notations.',
+
+	// Texture visualization
+	'Click the texture viewer to then zoom in and out with your mouse wheel.',
+	'Right click to drag texture viewport arround when zoomed in.',
+	'Uses the texture viewer\'s "copy to clipboard" to share your texture visualization to anyone.',
+];
+
+class TipView extends IView
 {
-	if (raw_pixel_data === null)
+	constructor(tip_id)
 	{
-		console.error(`failed to load ${resource_version_name}`);
-		return null;
+		super();
+		this.tip_id = tip_id;
 	}
 
-	var gl_format = translate_pixel_format(gl, resource_desc['Format']);
-	if (gl_format === null)
+	setup_html(parent_dom)
 	{
-		return null;
+		parent_dom.innerHTML = `
+			<div style="padding-top: 20%; width: 50%; margin: 0 auto; font-size: 14;">
+				<div>User tip:</div>
+				<div style="padding: 20px;">${k_tips[this.tip_id]}</div>
+				<div class="button_list" style="margin-left: auto; display: block; width: 100px;">
+						<a href="#display_tip(${(this.tip_id + 1) % k_tips.length});">Next</a>
+				</div>
+			</div>`;
 	}
 
-	var converted_data = null;
-	if (gl_format['src_type'] === gl.UNSIGNED_BYTE)
-	{
-		converted_data = new Uint8Array(raw_pixel_data);
-	}
-	else if (gl_format['src_type'] === gl.UNSIGNED_SHORT || gl_format['src_type'] === gl.HALF_FLOAT)
-	{
-		converted_data = new Uint16Array(raw_pixel_data);
-	}
-	else if (gl_format['src_type'] === gl.UNSIGNED_INT_10F_11F_11F_REV)
+	get navigations()
 	{
-		converted_data = new Uint32Array(raw_pixel_data);
+		return [`display_tip(${this.tip_id});`];
 	}
-	else if (gl_format['src_type'] === gl.FLOAT)
+}
+
+function display_tip(tip_id)
+{
+	if (tip_id === undefined)
 	{
-		converted_data = new Float32Array(raw_pixel_data);
+		tip_id = Math.floor(Math.random() * k_tips.length);
+		document.title = 'GPU Dump Viewer';
 	}
 	else
 	{
-		return null;
+		document.title = `User tip #${tip_id + 1}`;
 	}
 
-	const texture = gl.createTexture();
-	gl.bindTexture(gl.TEXTURE_2D, texture);
-	gl.texImage2D(
-		gl.TEXTURE_2D,
-		/* level = */ 0,
-		gl_format['internal_format'],
-	    resource_desc['ExtentX'],
-	    resource_desc['ExtentY'],
-	    /* border = */ 0,
-	    gl_format['src_format'],
-	    gl_format['src_type'],
-	    converted_data);
-	gl.texParameteri(gl.TEXTURE_2D, gl.TEXTURE_MAG_FILTER, gl.NEAREST);
-	gl.texParameteri(gl.TEXTURE_2D, gl.TEXTURE_MIN_FILTER, gl.NEAREST);
-	gl.texParameteri(gl.TEXTURE_2D, gl.TEXTURE_WRAP_S, gl.CLAMP_TO_EDGE);
-	gl.texParameteri(gl.TEXTURE_2D, gl.TEXTURE_WRAP_T, gl.CLAMP_TO_EDGE);
-	gl.bindTexture(gl.TEXTURE_2D, null);
-
-	return texture;
+	set_main_view(new TipView(tip_id));
 }
 
 
-function compile_screen_shader_program(gl, frag_code)
+// -------------------------------------------------------------------- DISPLAY CONSOLE VARIABLES
+
+class CVarsView extends IView
 {
-	const shader_code_prefix = `#version 300 es
+	constructor()
+	{
+		super();
+		this.load_cvars();
+	}
 
-precision highp float;
-precision highp int;
-precision highp sampler2D;
-precision highp usampler2D;`;
+	setup_html(parent_dom)
+	{
+		parent_dom.innerHTML = `
+			<div class="main_div">
+				<div class="pass_title">Console variables</div>
+				<div class="selection_list_search">
+					<input type="search" id="cvars_search_input" oninput="g_view.refresh_cvars();" onchange="g_view.refresh_cvars();" style="width: 100%;" placeholder="Search console variables..." />
+				</div>
+				<div id="cvars_pannel"></div>
+			</div>`;
 
-	var vert_code = `
-in vec3 coordinates;
-out highp vec2 tex_coordinates;
-void main(void) {
-	gl_Position = vec4(coordinates, 1.0);
-	tex_coordinates = coordinates.xy * 0.5 + 0.5;
-}`;
+		this.refresh_cvars();
+		document.title = 'Console variables';
+	}
 
-	var vert_shader = gl.createShader(gl.VERTEX_SHADER);
+	load_cvars()
 	{
-		gl.shaderSource(vert_shader, shader_code_prefix + vert_code);
-		gl.compileShader(vert_shader);
+		var cvars_csv = load_text_file('Base/ConsoleVariables.csv');
 
-		var compilation_status = gl.getShaderParameter(vert_shader, gl.COMPILE_STATUS);
-		if (!compilation_status)
+		this.cvars = [];
+		var cvars_set = new Set();
+		var csv_lines = cvars_csv.split('\n');
+		for (var i = 1; i < csv_lines.length - 1; i++)
 		{
-			var compilation_log = gl.getShaderInfoLog(vert_shader);
-			console.log('Fragment shader compiler log: ' + compilation_log);
-			return null;
+			var csv_line = csv_lines[i].split(',');
+			var entry = {};
+			entry['name'] = csv_line[0];
+			entry['type'] = csv_line[1];
+			entry['set_by'] = csv_line[2];
+			entry['value'] = csv_line[3];
+
+			if (!cvars_set.has(entry['name']))
+			{
+				cvars_set.add(entry['name']);
+				this.cvars.push(entry);
+			}
 		}
+
+		this.cvars.sort(function(a, b)
+		{
+			if (a['name'] < b['name'])
+			{
+				return -1;
+			}
+			else if (a['name'] > b['name'])
+			{
+				return 1;
+			}
+			return 0;
+		});
 	}
 
-	var frag_shader = gl.createShader(gl.FRAGMENT_SHADER);
+	refresh_cvars()
 	{
-		gl.shaderSource(frag_shader, shader_code_prefix + frag_code); 
-		gl.compileShader(frag_shader);
-
-		var compilation_status = gl.getShaderParameter(frag_shader, gl.COMPILE_STATUS);
-		if (!compilation_status)
+		var html = `
+			<table width="100%" class="pretty_table">
+				<tr class="header">
+					<td>Name</td>
+					<td width="15%">Value</td>
+					<td width="30px">Type</td>
+					<td width="100px">Set by</td>
+				</tr>`;
+
+		var cvar_search = document.getElementById('cvars_search_input').value.toLowerCase();
+
+		this.cvars.forEach(function(cvar)
 		{
-			var compilation_log = gl.getShaderInfoLog(frag_shader);
-			console.log('Fragment shader compiler log: ' + compilation_log);
-			document.getElementById('texture_visualization_code_log').value = 'Compilation failed:\n' + compilation_log;
-			return null;
-		}
+			var display = true;
+			if (cvar_search)
+			{
+				display = cvar['name'].toLowerCase().includes(cvar_search);
+			}
+
+			if (display)
+			{
+				html += `
+					<tr>
+						<td>${cvar['name']}</td>
+						<td>${cvar['value']}</td>
+						<td>${cvar['type']}</td>
+						<td>${cvar['set_by']}</td>
+					</tr>`;
+			}
+		});
+
+		html += `
+			</table>`;
+
+		document.getElementById('cvars_pannel').innerHTML = html;
 	}
 
-	var shader_program = gl.createProgram();
+	get navigations()
 	{
-		gl.attachShader(shader_program, vert_shader);
-		gl.attachShader(shader_program, frag_shader);
-		gl.linkProgram(shader_program);
+		return [`display_cvars();`];
+	}
 
-		var link_status = gl.getProgramParameter(shader_program, gl.LINK_STATUS);
-		if (!link_status)
-		{
-			var link_log = gl.getProgramInfoLog(shader_program);
-			console.log('Program compiler log: ' + link_log);
-			document.getElementById('texture_visualization_code_log').value = 'Link failed:\n' + link_log;
-		}
-		else
-		{
-			document.getElementById('texture_visualization_code_log').value = 'Compilation succeeded';
-		}
+	release()
+	{
+		this.cvars = null;
 	}
-	return shader_program;
 }
 
-function refresh_texture_visualization(resource_name, resource_version_name, raw_pixel_data)
+function display_cvars()
 {
-	var resource_desc = load_resource_desc(resource_name);
+	set_main_view(new CVarsView());
+}
+
 
-	var canvas = document.getElementById('texture_visualization_canvas');
-	var gl = g_gl;
+// -------------------------------------------------------------------- DISPLAY LOG
 
-	if (!g_gl_texture)
+class LogView extends IView
+{
+	constructor()
 	{
-		g_gl_texture = create_texture_from_pixel_data(gl, resource_desc, resource_version_name, raw_pixel_data);
+		super();
+		this.log = load_text_file(`Base/${g_infos['Project']}.log`).split('\n');
 	}
 
-	var texture = g_gl_texture;
-	if (!texture)
+	setup_html(parent_dom)
 	{
-		return;
+		parent_dom.innerHTML = `
+			<div class="main_div">
+				<div class="pass_title">Log</div>
+				<div class="selection_list_search">
+					<input type="search" id="log_search_input" oninput="g_view.refresh_log();" onchange="g_view.refresh_log();" style="width: 100%;" placeholder="Search log..." />
+				</div>
+				<div id="log_pannel"></div>
+			</div>`;
+
+		this.refresh_log();
+		document.title = 'Log';
 	}
 
+	refresh_log()
+	{
+		var html = `
+			<table width="100%" class="pretty_table">`;
 
-	/*======== Defining and storing the geometry ===========*/
+		var log_search = document.getElementById('log_search_input').value.toLowerCase();
 
-	var vertices = [
-		-1.0,+1.0,0.0,
-		-1.0,-1.0,0.0,
-		+1.0,-1.0,0.0,
-		+1.0,+1.0,0.0, 
-	];
+		this.log.forEach(function(log_line)
+		{
+			var display = true;
+			if (log_search)
+			{
+				display = log_line.toLowerCase().includes(log_search);
+			}
 
-	var indices = [0, 1, 2, 2, 0, 3];
+			if (display)
+			{
+				html += `
+					<tr>
+						<td>${log_line}</td>
+					</tr>`;
+			}
+		});
 
-	var vertex_buffer = gl.createBuffer();
-	{
-		gl.bindBuffer(gl.ARRAY_BUFFER, vertex_buffer);
-		gl.bufferData(gl.ARRAY_BUFFER, new Float32Array(vertices), gl.STATIC_DRAW);
-		gl.bindBuffer(gl.ARRAY_BUFFER, null);
+		html += `
+			</table>`;
+
+		document.getElementById('log_pannel').innerHTML = html;
 	}
 
-	var Index_Buffer = gl.createBuffer();
+	get navigations()
 	{
-		gl.bindBuffer(gl.ELEMENT_ARRAY_BUFFER, Index_Buffer);
-		gl.bufferData(gl.ELEMENT_ARRAY_BUFFER, new Uint16Array(indices), gl.STATIC_DRAW);
-		gl.bindBuffer(gl.ELEMENT_ARRAY_BUFFER, null);
+		return [`display_log();`];
 	}
 
-	/*================ Shaders ====================*/
-
-	var user_frag_code = document.getElementById('texture_visualization_code_input').value;
-	var shader_program = compile_screen_shader_program(gl, user_frag_code);
-
-	if (!shader_program)
+	release()
 	{
-		return;
+		this.log = null;
 	}
+}
 
-	// Save the user code.
-	g_shader_code_dict[resource_version_name] = user_frag_code;
+function display_log()
+{
+	set_main_view(new LogView());
+}
 
-	gl.clearColor(0.0, 0.0, 0.0, 1.0);
-	gl.clear(gl.COLOR_BUFFER_BIT);
 
-	{
-		gl.useProgram(shader_program);
-		gl.bindBuffer(gl.ARRAY_BUFFER, vertex_buffer);
-		gl.bindBuffer(gl.ELEMENT_ARRAY_BUFFER, Index_Buffer);
+// -------------------------------------------------------------------- DISPLAY TEXTURE
+
+var k_exotic_raw_channel_bit_depth = 0;
+var g_shader_code_dict = {};
+
+class TextureView extends ResourceView
+{
+	constructor(subresource_version_info, resource_desc)
+	{
+		super(subresource_version_info, resource_desc);
+		this.is_ready = false;
+		this.raw_texture_data = null;
+		this.subresource_desc = null;
+		this.release_gl();
+		this.viewport_width = 1;
+		this.pixel_scaling = 0;
+		this.viewport_selected = false;
+		this.is_draging_canvas = false;
+		this.display_mode = 'Visualization';
+	}
+
+	setup_html(parent_dom)
+	{
+		var subresource_extent = this.get_subresource_extent();
+
+		var dpi = window.devicePixelRatio || 1;
+
+		var canvas_rendering_res_x = subresource_extent['x'];
+		var canvas_rendering_res_y = subresource_extent['y'];
+
+		var canvas_display_w = canvas_rendering_res_x / dpi;
+		var canvas_display_h = canvas_rendering_res_y / dpi;
+
+		var resource_info_htmls = '';
+		{
+			resource_info_htmls += `
+				<table width="100%" class="pretty_table resource_desc">`;
+			
+			for (var key in this.resource_desc)
+			{
+			    if (this.resource_desc.hasOwnProperty(key)){
+					resource_info_htmls += `
+						<tr>
+							<td>${key}</td>
+							<td>${this.resource_desc[key]}</td>
+						</tr>`;
+			    }
+			}
+
+			resource_info_htmls += `
+				</table>`;
+		}
+
+		var title = `${this.resource_desc['Name']} ${this.subresource_version_info['subresource']}`;
+		if (this.subresource_version_info['draw'] >= 0)
+		{
+			title += ` draw:${this.subresource_version_info['draw']}`;
+		}
+
+		var html = `
+			<div class="main_div">
+				<div class="selection_list_title">Texture visualization: ${title}</div>
+				<div id="canvas_outter" style="margin-bottom: 10px;">
+					<div id="canvas_header">
+						<div class="button_list" id="texture_visualization_change_pixel_scaling"><!---
+							---><input type="button" onclick="g_view.resource_view.change_pixel_scaling(this);" value="Fit"/><!--- 
+							---><input type="button" onclick="g_view.resource_view.change_pixel_scaling(this);" value="1:1"/><!--- 
+							---><input type="button" onclick="g_view.resource_view.change_pixel_scaling(this);" value="2:1"/><!--- 
+							---><input type="button" onclick="g_view.resource_view.change_pixel_scaling(this);" value="4:1"/><!--- 
+							---><input type="button" onclick="g_view.resource_view.change_pixel_scaling(this);" value="8:1"/><!--- 
+							---><input type="button" onclick="g_view.resource_view.change_pixel_scaling(this);" value="16:1"/>
+						</div>
+						<div class="button_list">
+							<input type="button" onclick="g_view.resource_view.copy_canvas_to_clipboard();" value="Copy to clipboard"/>
+						</div>
+						<div class="button_list" id="texture_visualization_change_display_modes"><!---
+							---><input type="button" onclick="g_view.resource_view.change_display_mode(this);" value="Visualization"/><!---
+							---><input type="button" onclick="g_view.resource_view.change_display_mode(this);" value="NaN"/><!---
+							---><input type="button" onclick="g_view.resource_view.change_display_mode(this);" value="Inf"/>
+						</div>
+					</div>
+					<div id="canvas_viewport">
+						<canvas
+							id="texture_visualization_canvas"
+							width="${canvas_rendering_res_x}"
+							height="${canvas_rendering_res_y}"
+							style="width: ${canvas_display_w}px; height: ${canvas_display_h}px;"></canvas>
+					</div>
+					<div id="canvas_footer">
+						<table>
+							<tr>
+								<td>Cursor texel pos:</td>
+								<td id="canvas_hover_texel_info"></td>
+							</tr>
+							<tr>
+								<td>Selected texel pos:</td>
+								<td id="canvas_selected_texel_info"></td>
+							</tr>
+						</table>
+					</div>
+				</div>
+				<table width="100%">
+					<tr>
+						<td width="50%">
+							<div class="selection_list_title">Texture descriptor</div>
+							${resource_info_htmls}
+						</td>
+						<td style="width: 50%;">
+							<textarea id="texture_visualization_code_input" oninput="g_view.resource_view.shader_user_code_change();" onchange="g_view.resource_view.shader_user_code_change();"></textarea>
+							<textarea id="texture_visualization_code_log" readonly></textarea>
+							<div style="padding: 3px 10px;"><a class="external_link" href="https://www.khronos.org/files/webgl20-reference-guide.pdf">WebGL 2.0's quick reference card</a></div>
+						</td>
+					</tr>
+				</table>
+			</div>`;
+
+		parent_dom.innerHTML = html;
+
+		update_value_selection(document.getElementById('texture_visualization_change_pixel_scaling'), `Fit`);
+		update_value_selection(document.getElementById('texture_visualization_change_display_modes'), `Visualization`);
+
+		// Init WebGL 2.0
+		this.init_gl();
+
+		// Translate the descriptor of the subresource.
+		this.subresource_desc = this.translate_subresource_desc();
+		if (this.subresource_desc === null)
+		{
+			return;
+		}
+
+		// Setup mouse motion callback
+		{
+			var texture_view = this;
+
+			this.canvas.onclick = function(event) { texture_view.canvas_click(event); }
+			this.canvas.onmouseenter = function(event) { texture_view.canvas_onmousemove(event); };
+			this.canvas.onmousemove = function(event) { texture_view.canvas_onmousemove(event); };
+			this.canvas.onmousedown = function(event) { texture_view.canvas_onmousedown(event); };
+			this.canvas.onmouseup = function(event) { texture_view.canvas_onmouseup(event); };
+			this.canvas.oncontextmenu = function(event) { return false; };
+
+			document.getElementById('canvas_viewport').oncontextmenu = function(event) { return false; };
+
+			document.getElementById('canvas_outter').onclick = function(event) { texture_view.onclick_canvas_outter(); };
+			document.getElementById('canvas_outter').onmouseleave = function(event) { texture_view.set_select_viewport(false); };
+		}
+
+		if (this.subresource_desc['raw_channel_format'] == 'float')
+		{
+			update_value_selection(document.getElementById('texture_visualization_change_display_modes'), this.display_mode);
+		}
+		else
+		{
+			document.getElementById('texture_visualization_change_display_modes').style.display = 'none';
+		}
+
+		// Fill out the default visualization shader code in the textarea for the user to customise.
+		{
+			var user_shader_code_dom = document.getElementById('texture_visualization_code_input');
+
+			if (this.shader_code_saving_key in g_shader_code_dict)
+			{
+				user_shader_code_dom.value = g_shader_code_dict[this.shader_code_saving_key];
+			}
+			else
+			{
+				if (!this.subresource_desc)
+				{
+					user_shader_code_dom.value = `// Sorry but ${this.resource_desc['Format']} is not supported for display :(`;
+				}
+				else
+				{
+					user_shader_code_dom.value = this.subresource_desc['webgl_pixel_shader_public'];
+				}
+			}
+		}
+
+		var texture_view = this;
+		var subresource_version_name = get_subresource_unique_version_name(this.subresource_version_info);
+		load_binary_file(`Resources/${subresource_version_name}.bin`, function(raw_texture_data)
+		{
+			if (raw_texture_data)
+			{
+				texture_view.is_ready = true;
+				texture_view.raw_texture_data = raw_texture_data;
+				texture_view.resize_canvas(texture_view.pixel_scaling);
+				texture_view.process_texture_data_for_visualization();
+				texture_view.refresh_texture_visualization();
+			}
+			else
+			{
+				texture_view.error(`Error: Resources/${subresource_version_name}.bin couldn't be loaded.`);
+			}
+		});
+
+		document.title = `${g_view.pass_data['EventName']} - ${this.resource_desc['Name']} ${this.subresource_version_info['subresource']}`;
+	}
+
+	error(error_msg)
+	{
+		console.error(error_msg);
+		document.getElementById('canvas_viewport').innerHTML = `
+			<div class="error_msg">
+				${error_msg}
+			</div>`;
+		document.getElementById('texture_visualization_code_input').readOnly = true;
+	}
+
+	image_load_uint(texel_x, texel_y)
+	{
+		console.assert(this.is_ready);
+
+		var extent = this.get_subresource_extent();
+		var pixel_data_offset = texel_x + (extent['y'] - 1 - texel_y) * extent['x'];
+
+		var texel_value = [0];
+		for (var c = 1; c < this.subresource_desc['webgl_channel_count']; c++)
+		{
+			texel_value.push(0);
+		}
+
+		if (this.subresource_desc['webgl_src_type'] === this.gl.UNSIGNED_SHORT_5_6_5)
+		{
+			var data = new Uint16Array(this.raw_texture_data, pixel_data_offset * 2, 1);
+
+			texel_value[0] = (data[0] >>  0) & 0x1F;
+			texel_value[1] = (data[0] >>  5) & 0x3F;
+			texel_value[2] = (data[0] >> 11) & 0x1F;
+		}
+		else if (this.subresource_desc['webgl_src_type'] === this.gl.UNSIGNED_INT_10F_11F_11F_REV)
+		{
+			var data = new Uint32Array(this.raw_texture_data, pixel_data_offset * 4, 1);
+
+			texel_value[0] = (data[0] >>  0) & 0x7FF;
+			texel_value[1] = (data[0] >> 11) & 0x7FF;
+			texel_value[2] = (data[0] >> 22) & 0x3FF;
+		}
+		else if (this.subresource_desc['webgl_src_type'] === this.gl.UNSIGNED_INT_2_10_10_10_REV)
+		{
+			var data = new Uint32Array(this.raw_texture_data, pixel_data_offset * 4, 1);
+
+			texel_value[0] = (data[0] >>  0) & 0x3FF;
+			texel_value[1] = (data[0] >> 10) & 0x3FF;
+			texel_value[2] = (data[0] >> 20) & 0x3FF;
+			texel_value[3] = (data[0] >> 30) &   0x3;
+		}
+		else
+		{
+			var data = null;
+			if (this.subresource_desc['raw_channel_bit_depth'] === 8)
+			{
+				data = new Uint8Array(this.raw_texture_data, this.subresource_desc['webgl_channel_count'] * pixel_data_offset * 1, this.subresource_desc['webgl_channel_count']);
+			}
+			else if (this.subresource_desc['raw_channel_bit_depth'] === 16)
+			{
+				data = new Uint16Array(this.raw_texture_data, this.subresource_desc['webgl_channel_count'] * pixel_data_offset * 2, this.subresource_desc['webgl_channel_count']);
+			}
+			else if (this.subresource_desc['raw_channel_bit_depth'] === 32)
+			{
+				data = new Uint32Array(this.raw_texture_data, this.subresource_desc['webgl_channel_count'] * pixel_data_offset * 4, this.subresource_desc['webgl_channel_count']);
+			}
+			else
+			{
+				return;
+			}
+
+			for (var c = 0; c < this.subresource_desc['webgl_channel_count']; c++)
+			{
+				texel_value[c] = data[this.subresource_desc['raw_channel_swizzling'][c]];
+			}
+		}
+
+		return texel_value;
+	}
+
+	decode_texel_uint(value)
+	{
+		console.assert(this.is_ready);
+
+		value = [...value];
+
+		if (this.subresource_desc['ue_pixel_format'] === 'PF_FloatRGB' || this.subresource_desc['ue_pixel_format'] === 'PF_FloatR11G11B10')
+		{
+			value[0] = decode_float11(value[0]);
+			value[1] = decode_float11(value[1]);
+			value[2] = decode_float10(value[2]);
+		}
+		else if (this.subresource_desc['ue_pixel_format'] === 'PF_A2B10G10R10')
+		{
+			value[0] = value[0] / 1023.0;
+			value[1] = value[1] / 1023.0;
+			value[2] = value[2] / 1023.0;
+			value[3] = value[3] / 3.0;
+		}
+		else if (this.subresource_desc['ue_pixel_format'] === 'PF_R5G6B5_UNORM')
+		{
+			value[0] = value[0] / 31.0;
+			value[1] = value[1] / 63.0;
+			value[2] = value[2] / 31.0;
+		}
+		else if (this.subresource_desc['raw_channel_format'] == 'uint')
+		{
+			// NOP
+		}
+		else if (this.subresource_desc['raw_channel_format'] == 'float')
+		{
+			for (var c = 0; c < value.length; c++)
+			{
+				if (this.subresource_desc['raw_channel_bit_depth'] == 16)
+				{
+					value[c] = decode_float16(value[c]);
+				}
+				else if (this.subresource_desc['raw_channel_bit_depth'] == 32)
+				{
+					value[c] = decode_float32(value[c]);
+				}
+				else
+				{
+					console.error('Unknown float bit depth');
+				}
+			}
+		}
+		else if (this.subresource_desc['raw_channel_format'] == 'unorm')
+		{
+			var divide = Number((BigInt(1) << BigInt(this.subresource_desc['raw_channel_bit_depth'])) - 1n);
+
+			for (var c = 0; c < value.length; c++)
+			{
+				value[c] = Number(value[c]) / divide;
+			}
+		}
+		else
+		{
+			console.error('Unknown pixel format to convert');
+		}
+
+		return value;
+	}
+
+	update_texel_cursor_pos(texel_x, texel_y, parent_dom_name)
+	{
+		console.assert(this.is_ready);
+
+		var parent_dom = document.getElementById(parent_dom_name);
+		parent_dom.innerHTML = `<span style="width: 40px; display: inline-block; text-align: right;">${texel_x}</span> <span style="width: 40px; display: inline-block; text-align: right; margin-right: 60px;">${texel_y}</span> `;
+
+		// Find out pixel value.
+		{
+			var texel_raw = this.image_load_uint(texel_x, texel_y);
+			var texel_value = this.decode_texel_uint(texel_raw);
+
+			const channel_name = ['R', 'G', 'B', 'A'];
+			const channel_color = ['rgb(255, 38, 38)', 'rgb(38, 255, 38)', 'rgb(38, 187, 255)', 'white'];
+
+			var texel_string = [];
+			for (var c = 0; c < texel_value.length; c++)
+			{
+				texel_string.push(`${texel_value[c].toString()} (0x${texel_raw[c].toString(16)})`);
+			}
+
+			if (this.subresource_version_info['subresource'] == 'stencil')
+			{
+				texel_string[0] = '0b' + texel_value[0].toString(2).padStart(8, 0);
+			}
+
+			var html = ``;
+			for (var c = 0; c < texel_value.length; c++)
+			{
+				html += `<span style="margin-left: 20px; color: ${channel_color[c]};">${channel_name[c]}: <span style="min-width: 170px; display: inline-block; text-align: right;">${texel_string[c]}</span></span>`;
+			}
+
+			parent_dom.innerHTML += html;
+		}
+	}
+
+	get_texel_coordinate_of_mouse(event)
+	{
+		console.assert(this.is_ready);
+		
+		var rect = event.target.getBoundingClientRect();
+		var html_x = event.clientX - rect.left;
+		var html_y = event.clientY - rect.top;
+
+		var texel_x = Math.floor(html_x * px_string_to_int(this.canvas.width) / px_string_to_int(this.canvas.style.width));
+		var texel_y = Math.floor(html_y * px_string_to_int(this.canvas.height) / px_string_to_int(this.canvas.style.height));
+
+		return [texel_x, texel_y];
+	}
+
+	onclick_canvas_outter()
+	{
+		if (!this.is_ready)
+		{
+			return;
+		}
+
+		this.set_select_viewport(true);
+	}
+
+	onmousewheel(event)
+	{
+		if (!this.is_ready)
+		{
+			return;
+		}
+
+		if (this.viewport_selected == false)
+		{
+			// forward scroll event to parent
+			document.getElementById('main_right_pannel').scrollTop += event.deltaY;
+			return false;
+		}
+
+		var dpi = window.devicePixelRatio || 1;
+		var zooming_in = event.deltaY < 0.0;
+
+		var subresource_extent = this.get_subresource_extent();
+
+		// Find new pixel scaling to scale the viewport to.
+		var new_pixel_scaling = 0;
+		{
+			var viewport_width = (px_string_to_int(document.getElementById('canvas_viewport').style.width) - k_scroll_width) * dpi;
+
+			var min_pixel_scaling = Math.ceil(viewport_width / subresource_extent['x']);
+			if (min_pixel_scaling <= 1)
+			{
+				min_pixel_scaling = 1;
+			}
+			else if (min_pixel_scaling <= 2)
+			{
+				min_pixel_scaling = 2;
+			}
+			else if (min_pixel_scaling <= 4)
+			{
+				min_pixel_scaling = 4;
+			}
+			else if (min_pixel_scaling <= 8)
+			{
+				min_pixel_scaling = 8;
+			}
+			else if (min_pixel_scaling <= 16)
+			{
+				min_pixel_scaling = 16;
+			}
+
+			if (zooming_in)
+			{
+				if (this.pixel_scaling == 0)
+				{
+					if (min_pixel_scaling <= 16)
+					{
+						new_pixel_scaling = min_pixel_scaling;
+					}
+				}
+				else
+				{
+					new_pixel_scaling = this.pixel_scaling * 2;
+				}
+			}
+			else
+			{
+				if (this.pixel_scaling <= min_pixel_scaling)
+				{
+					new_pixel_scaling = 0;
+				}
+				else
+				{
+					new_pixel_scaling = this.pixel_scaling / 2;
+				}
+			}
+
+			new_pixel_scaling = Math.min(Math.max(new_pixel_scaling, 0), 16);
+		}
+
+		var [texel_x, texel_y] = this.get_texel_coordinate_of_mouse(event);
+
+		this.resize_canvas(new_pixel_scaling, texel_x, texel_y);
+
+		// Update the state of the zooming buttons
+		{
+			const pixel_scaling_names = {
+				0: 'Fit',
+				1: '1:1',
+				2: '2:1',
+				4: '4:1',
+				8: '8:1',
+				16: '16:1',
+			};
+			update_value_selection(document.getElementById('texture_visualization_change_pixel_scaling'), pixel_scaling_names[this.pixel_scaling]);
+		}
+
+		return false;
+	}
+
+	canvas_click(event)
+	{
+		if (!this.is_ready)
+		{
+			return;
+		}
+		
+		var subresource_extent = this.get_subresource_extent();
+
+		var [texel_x, texel_y] = this.get_texel_coordinate_of_mouse(event);
+
+		texel_x = Math.min(Math.max(texel_x, 0), subresource_extent['x'] - 1);
+		texel_y = Math.min(Math.max(texel_y, 0), subresource_extent['y'] - 1);
+
+		this.update_texel_cursor_pos(texel_x, texel_y, 'canvas_selected_texel_info');
+	}
+
+	canvas_onmousemove(event)
+	{
+		if (!this.is_ready)
+		{
+			return;
+		}
+		
+		if (this.is_draging_canvas)
+		{
+			this.drag_canvas(event);
+		}
+
+		var subresource_extent = this.get_subresource_extent();
+
+		var [texel_x, texel_y] = this.get_texel_coordinate_of_mouse(event);
+
+		texel_x = Math.min(Math.max(texel_x, 0), subresource_extent['x'] - 1);
+		texel_y = Math.min(Math.max(texel_y, 0), subresource_extent['y'] - 1);
+
+		this.update_texel_cursor_pos(texel_x, texel_y, 'canvas_hover_texel_info');
+	}
+
+	canvas_onmousedown(event)
+	{
+		if (!this.is_ready)
+		{
+			return;
+		}
+		
+		if (event.button == 2)
+		{
+			this.set_select_viewport(true);
+			this.is_draging_canvas = true;
+		}
+	}
+
+	canvas_onmouseup(event)
+	{
+		if (!this.is_ready)
+		{
+			return;
+		}
+		
+		if (event.button == 2)
+		{
+			this.is_draging_canvas = false;
+		}
+	}
+
+	drag_canvas(event)
+	{
+		if (!this.is_ready)
+		{
+			return;
+		}
+		
+		var canvas_viewport = document.getElementById('canvas_viewport');
+		canvas_viewport.scrollLeft -= event.movementX;
+		canvas_viewport.scrollTop -= event.movementY;
+	}
+
+	set_select_viewport(selected)
+	{
+		if (this.viewport_selected === selected)
+		{
+			return;
+		}
+
+		this.viewport_selected = selected;
+		if (this.viewport_selected)
+		{
+			var texture_view = this;
+			document.getElementById('main_right_pannel').onwheel = function(event) { return false; };
+			document.getElementById('canvas_outter').classList.add('selected');
+			document.getElementById('canvas_viewport').style.overflow = 'scroll';
+			this.canvas.onwheel = function(event) { return texture_view.onmousewheel(event); };
+		}
+		else
+		{
+			document.getElementById('main_right_pannel').onwheel = null;
+			document.getElementById('canvas_outter').classList.remove('selected');
+			document.getElementById('canvas_viewport').style.overflow = 'hidden';
+			this.canvas.onwheel = null;
+			this.is_draging_canvas = false;
+		}
+	}
+
+	resize(ctx)
+	{
+		var dpi = window.devicePixelRatio || 1;
+		var viewport_width = ctx.width - 100;
+		var viewport_width_no_scroll = viewport_width - k_scroll_width;
+
+		var subresource_extent = this.get_subresource_extent();
+		var canvas_rendering_res_x = subresource_extent['x'];
+		var canvas_rendering_res_y = subresource_extent['y'];
+
+		var canvas_outter = document.getElementById('canvas_outter');
+		var canvas_viewport = document.getElementById('canvas_viewport');
+		canvas_outter.style.width = `${viewport_width}px`;
+		canvas_viewport.style.width = `${viewport_width_no_scroll + k_scroll_width}px`;
+		canvas_viewport.style.height = `${viewport_width_no_scroll * canvas_rendering_res_y / canvas_rendering_res_x + k_scroll_width}px`;
+
+		this.viewport_width = viewport_width;
+		if (this.is_ready)
+		{
+			this.resize_canvas(this.pixel_scaling);
+		}
+	}
+
+	resize_canvas(new_pixel_scaling, fix_texel_x, fix_texel_y)
+	{
+		console.assert(this.is_ready);
+		var dpi = window.devicePixelRatio || 1;
+		var viewport_width = this.viewport_width;
+		var viewport_width_no_scroll = viewport_width - k_scroll_width;
+
+		var canvas_viewport = document.getElementById('canvas_viewport');
+		var canvas_viewport_rect = canvas_viewport.getBoundingClientRect();
+		var viewport_height = px_string_to_int(canvas_viewport.style.height);
+		var viewport_height_no_scroll = viewport_height - k_scroll_width;
+
+		var subresource_extent = this.get_subresource_extent();
+		var canvas_rendering_res_x = subresource_extent['x'];
+		var canvas_rendering_res_y = subresource_extent['y'];
+
+		if (fix_texel_x === undefined || fix_texel_y === undefined)
+		{
+			fix_texel_x = canvas_rendering_res_x / 2;
+			fix_texel_y = canvas_rendering_res_y / 2;
+		}
+
+		// Compute the pixel coordinate of the fixed texel.
+		var fix_texel_pixel_pos_x;
+		var fix_texel_pixel_pos_y;
+		{
+			var rect = this.canvas.getBoundingClientRect();
+
+			fix_texel_pixel_pos_x = rect.left + rect.width * fix_texel_x / canvas_rendering_res_x - canvas_viewport_rect.left;
+			fix_texel_pixel_pos_y = rect.top + rect.height * fix_texel_y / canvas_rendering_res_y - canvas_viewport_rect.top;
+		}
+
+		// Compute new canvas dimensions
+		var canvas_display_w = new_pixel_scaling * canvas_rendering_res_x / dpi;
+		var canvas_display_h = new_pixel_scaling * canvas_rendering_res_y / dpi;
+		if (new_pixel_scaling == 0)
+		{
+			canvas_display_w = viewport_width_no_scroll;
+			canvas_display_h = viewport_width_no_scroll * canvas_rendering_res_y / canvas_rendering_res_x;
+		}
+
+		// Compute new canvas scroll such that the fixed texel ends up at the same pixel coordinate.
+		var canvas_scroll_x;
+		var canvas_scroll_y;
+		{
+			canvas_scroll_x = fix_texel_x * canvas_display_w / canvas_rendering_res_x - fix_texel_pixel_pos_x;
+			canvas_scroll_y = fix_texel_y * canvas_display_h / canvas_rendering_res_y - fix_texel_pixel_pos_y;
+		}
+
+		this.canvas.style.width = `${canvas_display_w}px`;
+		this.canvas.style.height = `${canvas_display_h}px`;
+		canvas_viewport.scrollLeft = canvas_scroll_x;
+		canvas_viewport.scrollTop = canvas_scroll_y;
+		this.pixel_scaling = new_pixel_scaling;
+	}
+
+	init_gl()
+	{
+		var gl_ctx_attributes = {
+			antialias: false,
+  			depth: false
+		};
+
+		// Init WebGL 2.0
+		this.canvas = document.getElementById('texture_visualization_canvas');
+
+		var gl_ctx_attributes = {
+			alpha: false,
+  			// premultipliedAlpha: true, // TODO
+			antialias: false,
+  			depth: false,
+  			stencil: false,
+  			powerPreference: "low-power",
+  			preserveDrawingBuffer: true,
+  			xrCompatible: false,
+		};
+
+		var gl = this.canvas.getContext('webgl2', gl_ctx_attributes);
+		this.gl = gl;
+
+		// Init WebGL extensions
+		{
+			var available_extensions = this.gl.getSupportedExtensions();
+			var required_extensions = ['EXT_texture_norm16'];
+
+			this.gl_ext = {};
+
+			var gl = this.gl;
+			var gl_ext = this.gl_ext;
+			required_extensions.forEach(function(ext_name)
+			{
+				gl_ext[ext_name] = gl.getExtension(ext_name);
+			});
+		}
+
+		// Create vertex buffer
+		{
+			var vertices = [
+				-1.0,+1.0,0.0,
+				-1.0,-1.0,0.0,
+				+1.0,-1.0,0.0,
+				+1.0,+1.0,0.0, 
+			];
+
+			this.gl_vertex_buffer = gl.createBuffer();
+			gl.bindBuffer(gl.ARRAY_BUFFER, this.gl_vertex_buffer);
+			gl.bufferData(gl.ARRAY_BUFFER, new Float32Array(vertices), gl.STATIC_DRAW);
+			gl.bindBuffer(gl.ARRAY_BUFFER, null);
+		}
+
+		// Create vertex buffer
+		{
+			var indices = [0, 1, 2, 2, 0, 3];
+
+			this.gl_index_buffer = gl.createBuffer();
+			gl.bindBuffer(gl.ELEMENT_ARRAY_BUFFER, this.gl_index_buffer);
+			gl.bufferData(gl.ELEMENT_ARRAY_BUFFER, new Uint16Array(indices), gl.STATIC_DRAW);
+			gl.bindBuffer(gl.ELEMENT_ARRAY_BUFFER, null);
+		}
+
+		// Create simple shader program for NaN display mode
+		{
+			var frag_code = `
+uniform sampler2D texture0;
+
+in vec2 uv;
+out vec4 display;
+
+void main(void) {
+	ivec2 texel_coord = ivec2(uv * vec2(textureSize(texture0, 0)));
+	display = texelFetch(texture0, texel_coord, 0);
+}`;
+			this.gl_simple_program = this.compile_screen_shader_program(frag_code);
+		}
+
+		// Create simple shader program for NaN display mode
+		{
+			var frag_code = `
+uniform sampler2D texture0;
+
+in vec2 uv;
+out vec4 display;
+
+void main(void)
+{
+	ivec2 texel_coord = ivec2(uv * vec2(textureSize(texture0, 0)));
+	vec4 raw = texelFetch(texture0, texel_coord, 0);
+
+	display.r = (isnan(raw.r) || isnan(raw.g) || isnan(raw.b) || isnan(raw.a)) ? 1.0 : 0.0;
+	display.g = 0.0;
+	display.b = 0.0;
+	display.a = 1.0;
+}`;
+			this.gl_nan_program = this.compile_screen_shader_program(frag_code);
+		}
+
+		// Create simple shader program for Inf display mode
+		{
+			var frag_code = `
+uniform sampler2D texture0;
+
+in vec2 uv;
+out vec4 display;
+
+void main(void)
+{
+	ivec2 texel_coord = ivec2(uv * vec2(textureSize(texture0, 0)));
+	vec4 raw = texelFetch(texture0, texel_coord, 0);
+
+	display.r = (isinf(raw.r) || isinf(raw.g) || isinf(raw.b) || isinf(raw.a)) ? 1.0 : 0.0;
+	display.g = 0.0;
+	display.b = 0.0;
+	display.a = 1.0;
+}`;
+			this.gl_inf_program = this.compile_screen_shader_program(frag_code);
+		}
+	}
+
+	release_gl()
+	{
+		this.canvas = null;
+		this.gl = null;
+		this.gl_ext = null;
+		this.gl_vertex_buffer = null;
+		this.gl_index_buffer = null;
+		this.gl_simple_program = null;
+		this.gl_nan_program = null;
+		this.gl_textures = null;
+	}
+
+	release()
+	{
+		this.is_ready = false;
+		this.raw_texture_data = null;
+		this.release_gl();
+		this.set_select_viewport(false);
+	}
+
+	shader_user_code_change()
+	{
+		if (!this.is_ready)
+		{
+			return;
+		}
+		
+		this.refresh_texture_visualization();
+	}
+
+	copy_canvas_to_clipboard()
+	{
+		if (!this.is_ready)
+		{
+			return;
+		}
+		
+		var texture_view = this;
+		this.canvas.toBlob(function(image_blob) {
+			var pass_data = g_view.pass_data;
+
+			var is_output_resource = pass_data['OutputResources'].includes(get_subresource_unique_name(texture_view.subresource_version_info));
+
+			var text = '';
+			text += `Dump:\n\t${g_infos['Project']} - ${g_infos['Platform']} - ${g_infos['RHI']} - ${g_infos['BuildVersion']} - ${g_infos['DumpTime']}\n\t${get_current_navigation()}\n`
+			text += `Pass:\n\t${pass_data['EventName']}\n`;
+
+			{
+				var name = texture_view.resource_desc['Name'];
+				if (texture_view.subresource_version_info['subresource'])
+				{
+					name = `${texture_view.resource_desc['Name']} ${texture_view.subresource_version_info['subresource']}`;
+				}
+
+				text += `${is_output_resource ? 'Output' : 'Input'} resource:\n\t${name}\n`;
+			}
+
+			if (texture_view.subresource_version_info['draw'] >= 0)
+			{
+				text += `Draw ${texture_view.subresource_version_info['draw']}:\n\t${g_view.pass_draws_data[texture_view.subresource_version_info['draw']]['DrawName']}\n`;
+			}
+
+			var text_blob = new Blob([text], { type: 'text/plain' });
+			var clipboard_data = {
+				[text_blob.type]: text_blob,
+				[image_blob.type]: image_blob,
+			};
+
+			navigator.clipboard.write([new ClipboardItem(clipboard_data)]);
+		}, 'image/png');
+	}
+
+	get shader_code_saving_key()
+	{
+		var name = '';
+		if (this.resource_desc['Name'])
+		{
+			name = this.resource_desc['Name'];
+		}
+		else
+		{
+			name = this.subresource_version_info['resource'];
+		}
+
+		if (this.subresource_version_info['subresource'] && !this.subresource_version_info['subresource'].startsWith('mip'))
+		{
+			name = `${name}.${this.subresource_version_info['subresource']}`;
+		}
+
+		return name;
+	}
+
+	get_default_texture_shader_code()
+	{
+		if (this.shader_code_saving_key in g_shader_code_dict)
+		{
+			return g_shader_code_dict[this.shader_code_saving_key];
+		}
+
+		if (!this.subresource_desc)
+		{
+			return `// Sorry but ${this.resource_desc['Format']} is not supported for display :(`;
+		}
+
+		return this.subresource_desc['webgl_pixel_shader'];
+	}
+
+	translate_subresource_desc()
+	{
+		var gl = this.gl;
+		var ue_pixel_format = this.resource_desc['Format'];
+
+		// Missing formats:
+		//PF_BC6H
+		//PF_R8G8B8A8_SNORM
+		//PF_R16G16B16A16_UNORM
+		//PF_R16G16B16A16_SNORN
+
+		if (this.resource_desc['NumSamples'] > 1)
+		{
+			this.error(`MSAA is not supported for visualization.`);
+			return null;
+		}
+
+		var webgl_texture_internal_format = null;
+		var webgl_texture_src_format = null;
+		var webgl_texture_src_type = null;
+		var webgl_texture_count = 1;
+		var raw_channel_swizzling = [0, 1, 2, 3];
+		var raw_channel_bit_depth = 0;
+		var raw_channel_count = 0;
+		var raw_channel_format = 'exotic';
+
+		// 32bit floats
+		if (ue_pixel_format === 'PF_A32B32G32R32F')
+		{
+			webgl_texture_internal_format = gl.RGBA32F;
+			webgl_texture_src_format = gl.RGBA;
+			webgl_texture_src_type = gl.FLOAT;
+			raw_channel_bit_depth = 32;
+			raw_channel_count = 4;
+			raw_channel_format = 'float';
+		}
+		else if (ue_pixel_format === 'PF_G32R32F')
+		{
+			webgl_texture_internal_format = gl.RG32F;
+			webgl_texture_src_format = gl.RG;
+			webgl_texture_src_type = gl.FLOAT;
+			raw_channel_bit_depth = 32;
+			raw_channel_count = 2;
+			raw_channel_format = 'float';
+		}
+		else if (ue_pixel_format === 'PF_R32_FLOAT')
+		{
+			webgl_texture_internal_format = gl.R32F;
+			webgl_texture_src_format = gl.RED;
+			webgl_texture_src_type = gl.FLOAT;
+			raw_channel_bit_depth = 32;
+			raw_channel_count = 1;
+			raw_channel_format = 'float';
+		}
+		// 16bit floats
+		else if (ue_pixel_format === 'PF_FloatRGBA')
+		{
+			webgl_texture_internal_format = gl.RGBA16F;
+			webgl_texture_src_format = gl.RGBA;
+			webgl_texture_src_type = gl.HALF_FLOAT;
+			raw_channel_bit_depth = 16;
+			raw_channel_count = 4;
+			raw_channel_format = 'float';
+		}
+		else if (ue_pixel_format === 'PF_G16R16F' || ue_pixel_format === 'PF_G16R16F_FILTER')
+		{
+			webgl_texture_internal_format = gl.RG16F;
+			webgl_texture_src_format = gl.RG;
+			webgl_texture_src_type = gl.HALF_FLOAT;
+			raw_channel_bit_depth = 16;
+			raw_channel_count = 2;
+			raw_channel_format = 'float';
+		}
+		else if (ue_pixel_format === 'PF_R16F' || ue_pixel_format === 'PF_R16F_FILTER')
+		{
+			webgl_texture_internal_format = gl.R16F;
+			webgl_texture_src_format = gl.RED;
+			webgl_texture_src_type = gl.HALF_FLOAT;
+			raw_channel_bit_depth = 16;
+			raw_channel_count = 1;
+			raw_channel_format = 'float';
+		}
+		// 16bit unorms
+		else if (ue_pixel_format === 'PF_A16B16G16R16')
+		{
+			webgl_texture_internal_format = this.gl_ext['EXT_texture_norm16'].RGBA16_EXT;
+			webgl_texture_src_format = gl.RGBA;
+			webgl_texture_src_type = gl.UNSIGNED_SHORT;
+			raw_channel_bit_depth = 16;
+			raw_channel_count = 4;
+			raw_channel_format = 'unorm';
+		}
+		else if (ue_pixel_format === 'PF_G16R16')
+		{
+			webgl_texture_internal_format = this.gl_ext['EXT_texture_norm16'].RG16_EXT;
+			webgl_texture_src_format = gl.RG;
+			webgl_texture_src_type = gl.UNSIGNED_SHORT;
+			raw_channel_bit_depth = 16;
+			raw_channel_count = 2;
+			raw_channel_format = 'unorm';
+		}
+		else if (ue_pixel_format === 'PF_G16')
+		{
+			webgl_texture_internal_format = this.gl_ext['EXT_texture_norm16'].R16_EXT;
+			webgl_texture_src_format = gl.RED;
+			webgl_texture_src_type = gl.UNSIGNED_SHORT;
+			raw_channel_bit_depth = 16;
+			raw_channel_count = 1;
+			raw_channel_format = 'unorm';
+		}
+		// 8bit unorms
+		else if (ue_pixel_format === 'PF_B8G8R8A8' || ue_pixel_format === 'PF_R8G8B8A8')
+		{
+			webgl_texture_internal_format = gl.RGBA;
+			webgl_texture_src_format = gl.RGBA;
+			webgl_texture_src_type = gl.UNSIGNED_BYTE;
+			raw_channel_bit_depth = 8;
+			raw_channel_count = 4;
+			raw_channel_format = 'unorm';
+
+			if (ue_pixel_format === 'PF_B8G8R8A8')
+			{
+				raw_channel_swizzling = [2, 1, 0, 3];
+			}
+		}
+		else if (ue_pixel_format === 'PF_R8G8')
+		{
+			webgl_texture_internal_format = gl.RG8;
+			webgl_texture_src_format = gl.RG;
+			webgl_texture_src_type = gl.UNSIGNED_BYTE;
+			raw_channel_bit_depth = 8;
+			raw_channel_count = 2;
+			raw_channel_format = 'unorm';
+		}
+		else if (ue_pixel_format === 'PF_R8' || ue_pixel_format === 'PF_G8' || ue_pixel_format === 'PF_A8' || ue_pixel_format === 'PF_L8')
+		{
+			webgl_texture_internal_format = gl.R8;
+			webgl_texture_src_format = gl.RED;
+			webgl_texture_src_type = gl.UNSIGNED_BYTE;
+			raw_channel_bit_depth = 8;
+			raw_channel_count = 1;
+			raw_channel_format = 'unorm';
+		}
+		// 32bit uint
+		else if (ue_pixel_format === 'PF_R32_UINT' || ue_pixel_format === 'PF_R32_SINT')
+		{
+			webgl_texture_internal_format = gl.R8UI;
+			webgl_texture_src_format = gl.RED_INTEGER;
+			webgl_texture_src_type = gl.UNSIGNED_BYTE;
+			webgl_texture_count = 4;
+			raw_channel_bit_depth = 32;
+			raw_channel_count = 1;
+			raw_channel_format = 'uint';
+		}
+		else if (ue_pixel_format === 'PF_R32G32B32A32_UINT')
+		{
+			webgl_texture_internal_format = gl.RGBA8UI;
+			webgl_texture_src_format = gl.RGBA_INTEGER;
+			webgl_texture_src_type = gl.UNSIGNED_BYTE;
+			webgl_texture_count = 4;
+			raw_channel_bit_depth = 32;
+			raw_channel_count = 4;
+			raw_channel_format = 'uint';
+		}
+		else if (ue_pixel_format === 'PF_R32G32_UINT')
+		{
+			webgl_texture_internal_format = gl.RGBA8UI;
+			webgl_texture_src_format = gl.RGBA_INTEGER;
+			webgl_texture_src_type = gl.UNSIGNED_BYTE;
+			webgl_texture_count = 4;
+			raw_channel_bit_depth = 32;
+			raw_channel_count = 2;
+			raw_channel_format = 'uint';
+		}
+		// 16bit uint
+		else if (ue_pixel_format === 'PF_R16_UINT' || ue_pixel_format === 'PF_R16_SINT')
+		{
+			webgl_texture_internal_format = gl.R8UI;
+			webgl_texture_src_format = gl.RED_INTEGER;
+			webgl_texture_src_type = gl.UNSIGNED_BYTE;
+			webgl_texture_count = 2;
+			raw_channel_bit_depth = 16;
+			raw_channel_count = 1;
+			raw_channel_format = 'uint';
+		}
+		else if (ue_pixel_format === 'PF_R16G16B16A16_UINT' || ue_pixel_format === 'PF_R16G16B16A16_SINT')
+		{
+			webgl_texture_internal_format = gl.RGBA8UI;
+			webgl_texture_src_format = gl.RGBA_INTEGER;
+			webgl_texture_src_type = gl.UNSIGNED_BYTE;
+			webgl_texture_count = 2;
+			raw_channel_bit_depth = 16;
+			raw_channel_count = 4;
+			raw_channel_format = 'uint';
+		}
+		else if (ue_pixel_format === 'PF_R16G16_UINT')
+		{
+			webgl_texture_internal_format = gl.RGBA8UI;
+			webgl_texture_src_format = gl.RGBA_INTEGER;
+			webgl_texture_src_type = gl.UNSIGNED_BYTE;
+			webgl_texture_count = 2;
+			raw_channel_bit_depth = 16;
+			raw_channel_count = 2;
+			raw_channel_format = 'uint';
+		}
+		// 8bit uint
+		else if (ue_pixel_format === 'PF_R8G8B8A8_UINT')
+		{
+			webgl_texture_internal_format = gl.RGBA8UI;
+			webgl_texture_src_format = gl.RGBA_INTEGER;
+			webgl_texture_src_type = gl.UNSIGNED_BYTE;
+			raw_channel_bit_depth = 8;
+			raw_channel_count = 4;
+			raw_channel_format = 'uint';
+		}
+		else if (ue_pixel_format === 'PF_R8_UINT')
+		{
+			webgl_texture_internal_format = gl.R8UI;
+			webgl_texture_src_format = gl.RED_INTEGER;
+			webgl_texture_src_type = gl.UNSIGNED_BYTE;
+			raw_channel_bit_depth = 8;
+			raw_channel_count = 1;
+			raw_channel_format = 'uint';
+		}
+		// exotic bit depth
+		else if (ue_pixel_format === 'PF_FloatRGB' || ue_pixel_format === 'PF_FloatR11G11B10')
+		{
+			webgl_texture_internal_format = gl.R11F_G11F_B10F;
+			webgl_texture_src_format = gl.RGB;
+			webgl_texture_src_type = gl.UNSIGNED_INT_10F_11F_11F_REV;
+			raw_channel_bit_depth = k_exotic_raw_channel_bit_depth;
+			raw_channel_count = 3;
+			raw_channel_format = 'float';
+		}
+		else if (ue_pixel_format === 'PF_A2B10G10R10')
+		{
+			webgl_texture_internal_format = gl.RGB10_A2;
+			webgl_texture_src_format = gl.RGBA;
+			webgl_texture_src_type = gl.UNSIGNED_INT_2_10_10_10_REV;
+			raw_channel_bit_depth = k_exotic_raw_channel_bit_depth;
+			raw_channel_count = 4;
+			raw_channel_format = 'unorm';
+		}
+		else if (ue_pixel_format === 'PF_R5G6B5_UNORM')
+		{
+			webgl_texture_internal_format = gl.RGB565;
+			webgl_texture_src_format = gl.RGB;
+			webgl_texture_src_type = gl.UNSIGNED_SHORT_5_6_5;
+			raw_channel_bit_depth = k_exotic_raw_channel_bit_depth;
+			raw_channel_count = 3;
+			raw_channel_format = 'unorm';
+		}
+		// depth stencil
+		else if (ue_pixel_format === 'PF_DepthStencil' || ue_pixel_format === 'PF_ShadowDepth' || ue_pixel_format === 'PF_D24')
+		{
+			if (this.subresource_version_info['subresource'] == 'stencil')
+			{
+				webgl_texture_internal_format = gl.R8UI;
+				webgl_texture_src_format = gl.RED_INTEGER;
+				webgl_texture_src_type = gl.UNSIGNED_BYTE;
+				raw_channel_bit_depth = 8;
+				raw_channel_count = 1;
+				raw_channel_format = 'uint';
+			}
+			else
+			{
+				webgl_texture_internal_format = gl.R32F;
+				webgl_texture_src_format = gl.RED;
+				webgl_texture_src_type = gl.FLOAT;
+				raw_channel_bit_depth = 32;
+				raw_channel_count = 1;
+				raw_channel_format = 'float';
+			}
+		}
+		else
+		{
+			this.error(`Pixel format ${ue_pixel_format} is not supported for visualization.`);
+			return null;
+		}
+
+		var shader_float_vector = [null, 'float', 'vec2', 'vec3', 'vec4'];
+
+		var shader_sampler_type = 'sampler2D';
+		var shader_sampler_return_type = shader_float_vector;
+		var shader_display_operation = 'texel';
+		if (webgl_texture_src_format === gl.RED_INTEGER || webgl_texture_src_format === gl.RG_INTEGER || webgl_texture_src_format === gl.RGBA_INTEGER)
+		{
+			var divide = (BigInt(1) << BigInt(raw_channel_bit_depth)) - 1n;
+			shader_sampler_type = 'usampler2D';
+			shader_sampler_return_type = [null, 'uint', 'uvec2', 'uvec3', 'uvec4'];
+			shader_display_operation = `${shader_float_vector[raw_channel_count]}(${shader_display_operation}) / ${divide}.0`;
+		}
+
+		var shader_display_assignment = 'display.rgba';
+		var shader_fetch_return_operation = 'raw';
+		var webgl_channel_count = 0;
+		if (webgl_texture_src_format === gl.RED || webgl_texture_src_format === gl.RED_INTEGER)
+		{
+			shader_display_assignment = `display.r`;
+			shader_fetch_return_operation = `${shader_fetch_return_operation}.r`;
+			webgl_channel_count = 1;
+		}
+		else if (webgl_texture_src_format === gl.RG)
+		{
+			shader_display_assignment = `display.rg`;
+			shader_fetch_return_operation = `${shader_fetch_return_operation}.rg`;
+			webgl_channel_count = 2;
+		}
+		else if (webgl_texture_src_format === gl.RGB)
+		{
+			shader_display_assignment = `display.rgb`;
+			shader_fetch_return_operation = `${shader_fetch_return_operation}.rgb`;
+			webgl_channel_count = 3;
+		}
+		else if (webgl_texture_src_format === gl.RGBA || webgl_texture_src_format === gl.RGBA_INTEGER)
+		{
+			webgl_channel_count = 4;
+		}
+
+		var webgl_pixel_shader_private = ``;
+
+		var shader_fetch_operation = '(\n';
+		for (var i = 0; i < webgl_texture_count; i++)
+		{
+			webgl_pixel_shader_private += `uniform ${shader_sampler_type} texture${i};\n`;
+
+			shader_fetch_operation += `\t\t(texelFetch(texture${i}, texel_coord, 0) << ${8 * i})`;
+			if (i + 1 == webgl_texture_count)
+			{
+				shader_fetch_operation += `)`;
+			}
+			else
+			{
+				shader_fetch_operation += ` |\n`;
+			}
+		}
+
+		if (webgl_texture_count == 1)
+		{
+			shader_fetch_operation = `texelFetch(texture0, texel_coord , 0)`;
+		}
+
+		webgl_pixel_shader_private += `
+${shader_sampler_return_type[raw_channel_count]} fetchTexel(vec2 uv)
+{
+	ivec2 texel_coord = ivec2(uv * vec2(textureSize(texture0, 0)));
+	${shader_sampler_return_type[4]} raw = ${shader_fetch_operation};
+	return ${shader_fetch_return_operation};
+}
+
+in vec2 uv;
+out vec4 display;
+`;
+
+		var webgl_pixel_shader_public = `// WebGL 2.0 visualization shader for a ${ue_pixel_format} ${this.subresource_version_info['subresource']}
+
+${shader_sampler_return_type[raw_channel_count]} texel = fetchTexel(uv);
+${shader_display_assignment} = ${shader_display_operation};`;
+
+		var gl_format = {};
+		gl_format['webgl_internal_format'] = webgl_texture_internal_format;
+		gl_format['webgl_src_format'] = webgl_texture_src_format;
+		gl_format['webgl_src_type'] = webgl_texture_src_type;
+		gl_format['webgl_texture_count'] = webgl_texture_count;
+		gl_format['webgl_pixel_shader_public'] = webgl_pixel_shader_public;
+		gl_format['webgl_pixel_shader_private'] = webgl_pixel_shader_private;
+		gl_format['webgl_channel_count'] = webgl_channel_count;
+		gl_format['raw_channel_swizzling'] = raw_channel_swizzling;
+		gl_format['raw_channel_bit_depth'] = raw_channel_bit_depth;
+		gl_format['raw_channel_count'] = raw_channel_count;
+		gl_format['raw_channel_format'] = raw_channel_format;
+		gl_format['ue_pixel_format'] = ue_pixel_format;
+
+		return gl_format;
+	}
+
+	change_pixel_scaling(pixel_scaling_button)
+	{
+		if (!this.is_ready)
+		{
+			return;
+		}
+		
+		const pixel_scalings = {
+			'Fit': 0,
+			'1:1': 1,
+			'2:1': 2,
+			'4:1': 4,
+			'8:1': 8,
+			'16:1': 16,
+		};
+		this.resize_canvas(pixel_scalings[pixel_scaling_button.value]);
+		update_value_selection(document.getElementById('texture_visualization_change_pixel_scaling'), pixel_scaling_button.value);
+	}
+
+	change_display_mode(new_mode)
+	{
+		if (!this.is_ready)
+		{
+			return;
+		}
+		
+		if (new_mode.value == this.display_mode)
+		{
+			return;
+		}
+
+		this.display_mode = new_mode.value;
+
+		update_value_selection(document.getElementById('texture_visualization_change_display_modes'), this.display_mode);
+
+		if (this.display_mode == 'Visualization')
+		{
+			document.getElementById('texture_visualization_code_input').readOnly = false;
+		}
+		else
+		{
+			document.getElementById('texture_visualization_code_input').readOnly = true;
+		}
+
+		this.refresh_texture_visualization();
+	}
+
+	get_subresource_extent()
+	{
+		var extent = {};
+		extent['x'] = this.resource_desc['ExtentX'];
+		extent['y'] = this.resource_desc['ExtentY'];
+
+		if (this.subresource_version_info['subresource'].startsWith('mip'))
+		{
+			var mip_id = Number(this.subresource_version_info['subresource'].substring(3));
+
+			for (var i = 0; i < mip_id; i++)
+			{
+				extent['x'] = Math.floor(extent['x'] / 2);
+				extent['y'] = Math.floor(extent['y'] / 2);
+			}
+		}
+
+		return extent;
+	}
 
-		var coord = gl.getAttribLocation(shader_program, "coordinates");
-		gl.vertexAttribPointer(coord, 3, gl.FLOAT, false, 0, 0); 
-		gl.enableVertexAttribArray(coord);
+	create_gl_texture(webgl_internal_format, webgl_src_format, webgl_src_type, extent, converted_data)
+	{
+		console.assert(this.is_ready);
+		var gl = this.gl;
 
-		var tex_sampler = gl.getUniformLocation(shader_program, 'tex_sampler');
+		var texture = gl.createTexture();
 		gl.activeTexture(gl.TEXTURE0);
 		gl.bindTexture(gl.TEXTURE_2D, texture);
-		gl.uniform1i(tex_sampler, 0);
+		gl.texImage2D(
+			gl.TEXTURE_2D,
+			/* level = */ 0,
+			webgl_internal_format,
+		    extent['x'],
+		    extent['y'],
+		    /* border = */ 0,
+		    webgl_src_format,
+		    webgl_src_type,
+		    converted_data);
+		gl.texParameteri(gl.TEXTURE_2D, gl.TEXTURE_MAG_FILTER, gl.NEAREST);
+		gl.texParameteri(gl.TEXTURE_2D, gl.TEXTURE_MIN_FILTER, gl.NEAREST);
+		gl.texParameteri(gl.TEXTURE_2D, gl.TEXTURE_WRAP_S, gl.CLAMP_TO_EDGE);
+		gl.texParameteri(gl.TEXTURE_2D, gl.TEXTURE_WRAP_T, gl.CLAMP_TO_EDGE);
+		gl.bindTexture(gl.TEXTURE_2D, null);
+
+		return texture;
+	}
+
+	process_texture_data_for_visualization()
+	{
+		console.assert(this.is_ready);
+		var gl = this.gl;
+
+		if (this.raw_texture_data === null)
+		{
+			console.error(`failed to load ${get_subresource_unique_version_name(this.subresource_version_info)}`);
+			return null;
+		}
+
+		var extent = this.get_subresource_extent();
+		var webgl_array_size = this.subresource_desc['webgl_channel_count'] * extent['x'] * extent['y'];
+
+		// Reset the gl textures
+		this.gl_textures = [];
+
+		// Exotic pixel format
+		if (this.subresource_desc['raw_channel_bit_depth'] == k_exotic_raw_channel_bit_depth)
+		{
+			var gl = this.gl;
+			var original_converted_data = null;
+			if (this.subresource_desc['webgl_src_type'] === gl.UNSIGNED_SHORT_5_6_5)
+			{
+				original_converted_data = new Uint16Array(this.raw_texture_data);
+			}
+			else if (this.subresource_desc['webgl_src_type'] === gl.UNSIGNED_INT_10F_11F_11F_REV || this.subresource_desc['webgl_src_type'] === gl.UNSIGNED_INT_2_10_10_10_REV)
+			{
+				original_converted_data = new Uint32Array(this.raw_texture_data);
+			}
+			else
+			{
+				return null;
+			}
+
+			var texture = this.create_gl_texture(this.subresource_desc['webgl_internal_format'], this.subresource_desc['webgl_src_format'], this.subresource_desc['webgl_src_type'], extent, original_converted_data);
+			this.gl_textures.push(texture);
+		}
+		else if (this.subresource_desc['webgl_texture_count'] > 1)
+		{
+			var original_converted_data = null;
+			if (this.subresource_desc['raw_channel_bit_depth'] === 8)
+			{
+				original_converted_data = new Uint8Array(this.raw_texture_data);
+			}
+			else if (this.subresource_desc['raw_channel_bit_depth'] === 16)
+			{
+				original_converted_data = new Uint16Array(this.raw_texture_data);
+			}
+			else if (this.subresource_desc['raw_channel_bit_depth'] === 32)
+			{
+				original_converted_data = new Uint32Array(this.raw_texture_data);
+			}
+			else
+			{
+				return null;
+			}
+
+			for (var tex_id = 0; tex_id < this.subresource_desc['webgl_texture_count']; tex_id++)
+			{
+				var converted_data = new Uint8Array(webgl_array_size);
+
+				for (var y = 0; y < extent['y']; y++)
+				{
+					for (var x = 0; x < extent['x']; x++)
+					{
+						for (var c = 0; c < this.subresource_desc['webgl_channel_count']; c++)
+						{
+							var i = this.subresource_desc['webgl_channel_count'] * (x + extent['x'] * y);
+
+							var value = 0;
+							if (c < this.subresource_desc['raw_channel_count'])
+							{
+								var src = original_converted_data[i + this.subresource_desc['raw_channel_swizzling'][c]];
+								value = (src >> (8 * tex_id)) % 256;
+							}
+							converted_data[i + c] = value;
+						}
+					}
+				}
+
+				var texture = this.create_gl_texture(this.subresource_desc['webgl_internal_format'], this.subresource_desc['webgl_src_format'], this.subresource_desc['webgl_src_type'], extent, converted_data);
+				this.gl_textures.push(texture);
+			}
+		}
+		else
+		{
+			var gl = this.gl;
+			var converted_data = null;
+			var original_converted_data = null;
+			if (this.subresource_desc['webgl_src_type'] === gl.UNSIGNED_BYTE)
+			{
+				original_converted_data = new Uint8Array(this.raw_texture_data);
+				converted_data = new Uint8Array(webgl_array_size);
+			}
+			else if (this.subresource_desc['webgl_src_type'] === gl.UNSIGNED_SHORT || this.subresource_desc['webgl_src_type'] === gl.HALF_FLOAT)
+			{
+				original_converted_data = new Uint16Array(this.raw_texture_data);
+				converted_data = new Uint16Array(webgl_array_size);
+			}
+			else if (this.subresource_desc['webgl_src_type'] === gl.FLOAT)
+			{
+				original_converted_data = new Float32Array(this.raw_texture_data);
+				converted_data = new Float32Array(webgl_array_size);
+			}
+			else
+			{
+				return null;
+			}
+
+			for (var y = 0; y < extent['y']; y++)
+			{
+				for (var x = 0; x < extent['x']; x++)
+				{
+					for (var c = 0; c < this.subresource_desc['webgl_channel_count']; c++)
+					{
+						var i = this.subresource_desc['webgl_channel_count'] * (x + extent['x'] * y);
+
+						var value = 0;
+						if (c < this.subresource_desc['raw_channel_count'])
+						{
+							value = original_converted_data[i + this.subresource_desc['raw_channel_swizzling'][c]];
+						}
+						converted_data[i + c] = value;
+					}
+				}
+			}
+
+			var texture = this.create_gl_texture(this.subresource_desc['webgl_internal_format'], this.subresource_desc['webgl_src_format'], this.subresource_desc['webgl_src_type'], extent, converted_data);
+			this.gl_textures.push(texture);
+		}
+	}
+
+	compile_screen_shader_program(frag_code)
+	{
+		var gl = this.gl;
+
+		//console.log(frag_code);
+
+		const shader_std = `#version 300 es
+
+precision highp float;
+precision highp int;
+precision highp sampler2D;
+precision highp usampler2D;
+
+struct FloatInfos
+{
+	bool is_denormal;
+	bool is_infinity;
+	bool is_nan;
+};
+
+FloatInfos decodeFloat(float x)
+{
+	const uint mantissa_bit_count = 23u;
+	const uint exp_bit_count = 8u;
+
+	uint raw = floatBitsToUint(x);
+
+	uint sign_bit      = (raw >> (mantissa_bit_count + exp_bit_count)) & 0x1u;
+	uint mantissa_bits = (raw >> 0u)                     & ((0x1u << mantissa_bit_count) - 1u);
+	uint exp_bits      = (raw >> mantissa_bit_count)     & ((0x1u << exp_bit_count) - 1u);
+
+	bool is_max_exp = exp_bits == ((0x1u << exp_bit_count) - 1u);
+
+	FloatInfos infos;
+	infos.is_denormal = exp_bits == 0u;
+	infos.is_infinity = is_max_exp && mantissa_bits == 0u;
+	infos.is_nan = is_max_exp && mantissa_bits != 0u;
+
+	return infos;
+}
+
+#define isnan(x) decodeFloat(x).is_nan
+#define isinf(x) decodeFloat(x).is_infinity
+#define isdenorm(x) decodeFloat(x).is_denormal
+
+`;
+
+		var vert_code = `
+in vec3 coordinates;
+out highp vec2 uv;
+void main(void) {
+	gl_Position = vec4(coordinates, 1.0);
+	uv = coordinates.xy * 0.5 + 0.5;
+}`;
+
+		var vert_shader = gl.createShader(gl.VERTEX_SHADER);
+		{
+			gl.shaderSource(vert_shader, shader_std + vert_code);
+			gl.compileShader(vert_shader);
+
+			var compilation_status = gl.getShaderParameter(vert_shader, gl.COMPILE_STATUS);
+			if (!compilation_status)
+			{
+				var compilation_log = gl.getShaderInfoLog(vert_shader);
+				console.error('Fragment shader compiler log: ' + compilation_log);
+				return null;
+			}
+		}
+
+		var frag_shader = gl.createShader(gl.FRAGMENT_SHADER);
+		{
+			gl.shaderSource(frag_shader, shader_std + frag_code); 
+			gl.compileShader(frag_shader);
+
+			var compilation_status = gl.getShaderParameter(frag_shader, gl.COMPILE_STATUS);
+			if (!compilation_status)
+			{
+				var compilation_log = gl.getShaderInfoLog(frag_shader);
+				console.log(frag_code);
+				console.error('Fragment shader compiler log: ' + compilation_log);
+				document.getElementById('texture_visualization_code_log').value = 'Compilation failed:\n' + compilation_log;
+				return null;
+			}
+		}
+
+		var shader_program = gl.createProgram();
+		{
+			gl.attachShader(shader_program, vert_shader);
+			gl.attachShader(shader_program, frag_shader);
+			gl.linkProgram(shader_program);
+
+			var link_status = gl.getProgramParameter(shader_program, gl.LINK_STATUS);
+			if (!link_status)
+			{
+				var link_log = gl.getProgramInfoLog(shader_program);
+				console.error('Program compiler log: ' + link_log);
+				document.getElementById('texture_visualization_code_log').value = 'Link failed:\n' + link_log;
+			}
+			else
+			{
+				document.getElementById('texture_visualization_code_log').value = 'Compilation succeeded';
+			}
+		}
+		return shader_program;
+	}
+
+	refresh_texture_visualization()
+	{
+		console.assert(this.is_ready);
+		console.assert(this.gl_textures);
+
+		var shader_program = null;
+		if (this.display_mode == 'NaN')
+		{
+			shader_program = this.gl_nan_program;
+		}
+		else if (this.display_mode == 'Inf')
+		{
+			shader_program = this.gl_inf_program;
+		}
+		else
+		{
+			var user_frag_code = document.getElementById('texture_visualization_code_input').value;
+			shader_program = this.compile_screen_shader_program(this.subresource_desc['webgl_pixel_shader_private'] + 'void main(void) { display = vec4(0.0, 0.0, 0.0, 0.0);\n' + user_frag_code + `}`);
+
+			if (!shader_program)
+			{
+				return;
+			}
+
+			// Save the user code if compilation have succeeded.
+			g_shader_code_dict[this.shader_code_saving_key] = user_frag_code;
+		}
 
-		gl.viewport(0, 0, canvas.width, canvas.height);
-		gl.drawElements(gl.TRIANGLES, indices.length, gl.UNSIGNED_SHORT, 0);
+		var gl = this.gl;
+		gl.clearColor(0.0, 0.0, 0.0, 1.0);
+		gl.clear(gl.COLOR_BUFFER_BIT);
+
+		{
+			gl.useProgram(shader_program);
+			gl.bindBuffer(gl.ARRAY_BUFFER, this.gl_vertex_buffer);
+			gl.bindBuffer(gl.ELEMENT_ARRAY_BUFFER, this.gl_index_buffer);
+
+			var coord = gl.getAttribLocation(shader_program, "coordinates");
+			gl.vertexAttribPointer(coord, 3, gl.FLOAT, false, 0, 0); 
+			gl.enableVertexAttribArray(coord);
+
+			for (var tex_id = 0; tex_id < this.gl_textures.length; tex_id++)
+			{
+				gl.activeTexture(gl.TEXTURE0 + tex_id);
+				gl.bindTexture(gl.TEXTURE_2D, this.gl_textures[tex_id]);
+
+				var texture_sampler = gl.getUniformLocation(shader_program, 'texture' + tex_id);
+				gl.uniform1i(texture_sampler, tex_id);
+			}
+
+			gl.viewport(0, 0, this.canvas.width, this.canvas.height);
+			gl.drawElements(gl.TRIANGLES, 6, gl.UNSIGNED_SHORT, 0);
+		}
+	}
+} // class TextureView
+
+function display_texture(subresource_version_info)
+{
+	if (g_view.resource_view !== null && get_subresource_unique_version_name(g_view.resource_view.subresource_version_info) == get_subresource_unique_version_name(subresource_version_info))
+	{
+		return;
 	}
+
+	var resource_desc = load_resource_desc(subresource_version_info['resource']);
+	g_view.set_resource_view(new TextureView(subresource_version_info, resource_desc));
 }
 
 
 // -------------------------------------------------------------------- DISPLAY BUFFER
 
-var g_buffer_name = null;
-var g_buffer_version_name = null;
-var g_raw_buffer = null;
+const k_buffer_pixel_per_row = 20;
+const k_buffer_max_display_row_count = 50;
+const k_buffer_display_row_buffer = 50;
+
+
+class GenericBufferView extends ResourceView
+{
+	constructor(subresource_version_info, resource_desc)
+	{
+		super(subresource_version_info, resource_desc);
+		this.raw_buffer_data = null;
+		this.raw_buffer_data_path = null;
+		this.viewport_selected = false;
+		this.display_row_start = 0;
+		this.selected_address = -1;
+		this.structure_metadata = null;
+		this.display_elem_using_rows = resource_desc['NumElements'] == 1;
+	}
+
+	setup_html(parent_dom)
+	{
+		var html = '';
+
+		if (this.resource_desc['Desc'] != 'FRDGParameterStruct')
+		{
+			var resource_info_htmls = `
+				<table width="100%" class="resource_desc pretty_table">`;
+			
+			for (var key in this.resource_desc)
+			{
+			    if (this.resource_desc.hasOwnProperty(key))
+			    {
+			    	var value = this.resource_desc[key];
+
+			    	if (key == 'Metadata')
+			    	{
+			    		if (this.resource_desc[key] == k_null_json_ptr)
+			    		{
+			    			value = 'nullptr';
+			    		}
+			    		else if (this.structure_metadata)
+			    		{
+			    			value = `${this.structure_metadata['StructTypeName']} (${get_filename(this.structure_metadata['FileName'])}:${this.structure_metadata['FileLine']})`;
+			    		}
+			    	}
+
+					resource_info_htmls += `
+						<tr>
+							<td>${key}</td>
+							<td>${value}</td>
+						</tr>`;
+			    }
+			}
+
+			resource_info_htmls += `
+				</table>`;
+
+			html += `
+				<div class="main_div">
+					<div class="selection_list_title">Buffer descriptor: ${this.resource_desc['Name']}</div>
+					<div>
+						${resource_info_htmls}
+					</div>
+				</div>`;
+		}
+
+		html += `
+			<div class="main_div">
+				<div class="selection_list_title" style="width: auto;">Buffer visualization: ${this.resource_desc['Name']}</div>
+				<div id="buffer_outter">
+					<div id="buffer_viewport_header">
+						<table width="100%" id="buffer_visualization_format_outter">
+							<tr>
+								<td style="padding: 4px 20px 0 20px; text-align: right;" width="40px">Format:</td>
+								<td>
+									<input type="text" id="buffer_visualization_format" onchange="g_view.resource_view.refresh_buffer_visualization(true);" style="width: 100%;" />
+								</td>
+							</tr>
+						</table>
+						<div id="buffer_visualization_member_search_outter">
+							<input type="search" id="buffer_visualization_member_search" oninput="g_view.resource_view.refresh_members();" onchange="g_view.resource_view.refresh_members();" style="width: 100%;" placeholder="Search member..." />
+						</div>
+					</div>
+					<div id="buffer_viewport">
+						<div id="buffer_content_pannel"></div>
+					</div>
+				</div>
+			</div>`;
+
+		parent_dom.innerHTML = html;
+
+		if (this.structure_metadata)
+		{
+			// document.getElementById('buffer_visualization_format').value = `${this.structure_metadata['StructTypeName']} (${get_filename(this.structure_metadata['FileName'])}:${this.structure_metadata['FileLine']})`;
+			// document.getElementById('buffer_visualization_format').readOnly = true;
+			document.getElementById('buffer_visualization_format_outter').style.display = 'none';
+		}
+		else
+		{
+			var default_format = '';
+			{
+				var byte_per_element = this.resource_desc['BytesPerElement'];
+				var member_count = byte_per_element / 4;
+
+				var default_format = 'hex(uint)';
+
+				// DrawIndirect flags means this is certainly a 32bit uint.
+				if ('Usage' in this.resource_desc && this.resource_desc['Usage'].includes('DrawIndirect'))
+				{
+					default_format = 'uint';
+				}
+
+				var format_per_column = [];
+				for (var i = 0; i < member_count; i++)
+				{
+					format_per_column.push(default_format);
+				}
+
+				default_format = format_per_column.join(', ');
+			}
+
+			document.getElementById('buffer_visualization_format').value = default_format;
+			document.getElementById('buffer_visualization_member_search_outter').style.display = 'none';
+		}
+
+		// Setup mouse callbacks
+		{
+			var buffer_view = this;
+
+			document.getElementById('buffer_outter').onclick = function(event) { buffer_view.onclick_buffer_outter(event); };
+			document.getElementById('buffer_outter').onmouseleave = function(event) { buffer_view.set_select_viewport(false); };
+			document.getElementById('buffer_viewport').onclick = function(event) { buffer_view.onclick_buffer_viewport(event); };
+			document.getElementById('buffer_viewport').onscroll = function(event) { buffer_view.update_buffer_scroll(/* force_refresh = */ false); };
+		}
+
+		var buffer_view = this;
+		console.assert(this.raw_buffer_data_path);
+		load_binary_file(this.raw_buffer_data_path, function(raw_buffer_data)
+		{
+			if (raw_buffer_data)
+			{
+				buffer_view.raw_buffer_data = raw_buffer_data;
+				buffer_view.refresh_buffer_visualization(/* refresh header = */ true);
+			}
+			else
+			{
+				document.getElementById('buffer_content_pannel').innerHTML = `
+					Error: ${this.raw_buffer_data_path} couldn't be loaded.`;
+			}
+		});
+
+		document.title = `${g_view.pass_data['EventName']} - ${this.resource_desc['Name']}`;
+	}
+
+	refresh_members()
+	{
+		this.refresh_buffer_visualization(/* refresh_header = */ !this.display_elem_using_rows);
+	}
+
+	onclick_buffer_outter(event)
+	{
+		this.set_select_viewport(true);
+
+		if (this.display_elem_using_rows)
+		{
+			return;
+		}
+
+		var content_rect = document.getElementById('buffer_content_pannel').getBoundingClientRect();
+		var header_rect = document.getElementById('buffer_content_header').getBoundingClientRect();
+
+		var html_x = event.clientX - content_rect.left;
+		var html_y = event.clientY - content_rect.top - header_rect.height;
+
+		var new_selected_address = Math.floor(html_y / k_buffer_pixel_per_row);
+
+		if (new_selected_address != this.selected_address && new_selected_address < this.resource_desc['NumElements'])
+		{
+			this.selected_address = new_selected_address;
+			this.refresh_buffer_visualization();
+		}
+	}
+
+	onclick_buffer_viewport(event)
+	{
+		this.set_select_viewport(true);
+	}
+
+	set_select_viewport(selected)
+	{
+		if (this.viewport_selected === selected)
+		{
+			return;
+		}
+
+		this.viewport_selected = selected;
+		if (this.viewport_selected)
+		{
+			var buffer_view = this;
+			document.getElementById('buffer_outter').classList.add('selected');
+			document.getElementById('buffer_viewport').style.overflow = 'scroll';
+			//document.getElementById('buffer_viewport').style.margin = '0px 0px 0px 0px';
+		}
+		else
+		{
+			document.getElementById('buffer_outter').classList.remove('selected');
+			document.getElementById('buffer_viewport').style.overflow = 'hidden';
+			//document.getElementById('buffer_viewport').style.margin = `0px ${k_scroll_width}px ${k_scroll_width}px 0px`;
+		}
+	}
+
+	update_buffer_scroll(force_refresh)
+	{
+		var buffer_viewport = document.getElementById('buffer_viewport');
 
-function display_buffer(resource_name, resource_version_name)
-{
-	var resource_desc = load_resource_desc(resource_name);
+		var new_display_row_start = Math.max(2 * Math.floor((buffer_viewport.scrollTop / k_buffer_pixel_per_row) * 0.5), 0);
 
-	g_texture_name = resource_name;
-	g_texture_version_name = resource_version_name;
-	g_gl = null;
-	g_gl_texture = null;
-	g_raw_buffer = null;
-	g_buffer_name = null;
-	g_buffer_version_name = null;
+		if (new_display_row_start != this.display_row_start || force_refresh)
+		{
+			this.display_row_start = new_display_row_start;
+			this.refresh_buffer_visualization(/* refresh_header = */ false);
+		}
+	}
 
-	var default_format = '';
+	go_to_address()
 	{
-		var byte_per_element = resource_desc['BytesPerElement'];
-		var column_count = byte_per_element / 4;
+		var byte_per_element = this.resource_desc['BytesPerElement'];
 
-		var format_per_column = [];
-		for (var i = 0; i < column_count; i++)
+		var address = document.getElementById('buffer_address_input').value;
+		if (address == '')
 		{
-			format_per_column.push('hex');
+			if (this.selected_address != -1)
+			{
+				this.selected_address = -1;
+				this.update_buffer_scroll(/* force_refresh = */ true);
+			}
+			return;
 		}
 
-		default_format = format_per_column.join(', ');
+		var element_id = Math.floor(Number(address) / byte_per_element);
+		if (element_id < this.resource_desc['NumElements'])
+		{
+			this.selected_address = element_id;
+
+			var buffer_viewport = document.getElementById('buffer_viewport');
+			buffer_viewport.scrollTop = (element_id - 5) * k_buffer_pixel_per_row;
+
+			this.update_buffer_scroll(/* force_refresh = */ true);
+		}
 	}
 
-	var resource_info_htmls = '';
+	refresh_buffer_visualization(refresh_header)
 	{
-		resource_info_htmls += `
-			<table width="100%" class="resource_desc">`;
-		
-		for (var key in resource_desc){
-		    if (resource_desc.hasOwnProperty(key)){
-				resource_info_htmls += `
-					<tr>
-						<td>${key}</td>
-						<td>${resource_desc[key]}</td>
-					</tr>`;
-		    }
-		}
+		var byte_per_element = this.resource_desc['BytesPerElement'];
+		var num_element = this.resource_desc['NumElements'];
+
+		var member_count = 0;
+		var member_names = [];
+		var member_byte_offset = [];
+		var member_format = [];
+		if (this.structure_metadata)
+		{
+			var member_search = document.getElementById('buffer_visualization_member_search').value;
 
-		resource_info_htmls += `
-			</table>`;
-	}
+			iterate_structure_members(this.structure_metadata, function(it) {
+				var column_format = undefined;
+				if (it.base_type == 'UBMT_INT32')
+				{
+					column_format = 'int';
+				}
+				else if (it.base_type == 'UBMT_UINT32')
+				{
+					column_format = 'uint';
+				}
+				else if (it.base_type == 'UBMT_FLOAT32')
+				{
+					column_format = 'float';
+				}
+				else
+				{
+					return;
+				}
 
-	var html = '';
-	html += `
-		<div class="main_div">
-			<div class="selection_list_title">Buffer descriptor: ${resource_desc['Name']}</div>
-			<div>
-				${resource_info_htmls}
-			</div>
-		</div>
-		<div class="main_div">
-			<div class="selection_list_title" style="width: auto;">Buffer visualization: ${resource_desc['Name']}</div>
-			<div><input type="text" id="buffer_visualization_format" value="${default_format}" onchange="refresh_current_buffer_visualization();" style="width: 100%;" /></div>
-			<div id="buffer_content_pannel"></div>
-		</div>`;
+				if (member_search)
+				{
+					if (!it.cpp_name.toLowerCase().includes(member_search.toLowerCase()))
+					{
+						return;
+					}
+				}
 
-	document.getElementById('display_resource_pannel').innerHTML = html + get_resource_producer_html(resource_name);
+				const k_suffixes = ['.x', '.y', '.z', '.w'];
 
-	load_binary_file(`Resources/${resource_version_name}.bin`, function(raw_buffer_data)
-	{
-		if (raw_buffer_data)
-		{
-			g_raw_buffer = raw_buffer_data;
-			g_buffer_name = resource_name;
-			g_buffer_version_name = resource_version_name;
-			refresh_buffer_visualization(resource_name, resource_version_name, raw_buffer_data);
+				var relative_offset = 0;
+				for (var elem_id = 0; elem_id < Math.max(it.member['NumElements'], 1); elem_id++)
+				{
+					for (var row_id = 0; row_id < Math.max(it.member['NumRows'], 1); row_id++)
+					{
+						for (var column_id = 0; column_id < it.member['NumColumns']; column_id++)
+						{
+							var name = it.cpp_name;
+
+							if (it.member['NumElements'] > 0)
+							{
+								name += `[${elem_id}]`;
+							}
+							if (it.member['NumRows'] > 1)
+							{
+								name += k_suffixes[row_id];
+							}
+							if (it.member['NumColumns'] > 1)
+							{
+								name += k_suffixes[column_id];
+							}
+
+							member_count += 1;
+							member_byte_offset.push(it.byte_offset + relative_offset);
+							member_format.push(column_format);
+							member_names.push(name);
+
+							relative_offset += 4;
+						}
+					}
+				}
+			});
 		}
 		else
 		{
-			document.getElementById('buffer_content_pannel').innerHTML = `
-				Error: Resources/${resource_version_name}.bin couldn't be loaded.`;
-		}
-	});
-}
-
-function refresh_current_buffer_visualization()
-{
-	refresh_buffer_visualization(g_buffer_name, g_buffer_version_name, g_raw_buffer);
-}
+			member_count = byte_per_element / 4;
 
-function refresh_buffer_visualization(resource_name, buffer_visualization_format, raw_buffer_data)
-{
-	var resource_desc = load_resource_desc(resource_name);
-	var byte_per_element = resource_desc['BytesPerElement'];
-	var num_element = resource_desc['NumElements'];
-
-	var column_count = byte_per_element / 4;
+			var text_format = document.getElementById('buffer_visualization_format').value;
 
-	var column_byte_offset = [];
-	var buffer_format = [];
-	{
-		var text_format = document.getElementById('buffer_visualization_format').value;
+			var column_formats = text_format.split(',');
 
-		var column_formats = text_format.split(',');
+			var byte_offset = 0;
 
-		var byte_offset = 0;
+			column_formats.forEach(function(value, index) {
+				var format_to_use = value.trim();
 
-		column_formats.forEach(function(value, index) {
-			var format_to_use = value.trim();
+				member_byte_offset.push(byte_offset);
+				member_format.push(format_to_use);
+				member_names.push(format_to_use);
 
-			column_byte_offset.push(byte_offset);
-			buffer_format.push(format_to_use);
+				byte_offset += 4;
+			});
 
-			byte_offset += 4;
-		});
+			for (var i = member_format.length; i < member_count; i++)
+			{
+				var format_to_use = member_format[i % column_formats.length];
 
-		for (var i = buffer_format.length; i < column_count; i++)
-		{
-			var format_to_use = buffer_format[i % column_formats.length];
+				member_byte_offset.push(byte_offset);
+				member_format.push(format_to_use);
+				member_names.push(format_to_use);
 
-			column_byte_offset.push(byte_offset);
-			buffer_format.push(format_to_use);
+				byte_offset += 4;
+			}
 
-			byte_offset += 4;
 		}
 
-	}
-
-	var resource_content = '';
-	{
-		resource_content += `<table id="buffer_content_table">
-			<tr>
-				<td>Address</td>`;
-
-		for (var i = 0; i < column_count; i++)
+		var column_count = member_count;
+		var row_count = num_element;
+		if (this.display_elem_using_rows)
 		{
-			resource_content += `<td>${buffer_format[i]}</td>`;
+			if (this.structure_metadata)
+			{
+				column_count = 2;
+			}
+			else
+			{
+				column_count = 1;
+			}
+			row_count = member_count;
 		}
 
-		resource_content += `</tr>`;
+		var display_start_row_id = Math.max(this.display_row_start - k_buffer_display_row_buffer, 0);
+		var display_row_count = Math.min(row_count - display_start_row_id, k_buffer_max_display_row_count + 2 * k_buffer_display_row_buffer);
+		var display_end_row_id = display_start_row_id + display_row_count;
+
+		var scroll_top = document.getElementById('buffer_viewport').scrollTop;
 
-		for (var elem_id = 0; elem_id < num_element; elem_id++)
+		if (refresh_header)
 		{
-			resource_content += `
-				<tr>
-					<td>${elem_id * byte_per_element}</td>`;
+			var address_search_value = '';
+			if (document.getElementById('buffer_address_input'))
+			{
+				address_search_value = document.getElementById('buffer_address_input').value;
+			}
 
-			for (var i = 0; i < column_count; i++)
+			var classes = 'pretty_table';
+			if (this.display_elem_using_rows)
 			{
-				var byte_offset = byte_per_element * elem_id + column_byte_offset[i];
+				classes += ' display_elem_using_rows';
+			}
 
-				var value = null;
-				var display = buffer_format[i];
-				if (display == 'float')
-				{
-					value = (new Float32Array(raw_buffer_data, byte_offset, 1))[0];
-				}
-				else if (display == 'int')
+			var resource_content = `
+				<table id="buffer_content_table" class="${classes}" width="100%">
+					<tr class="header" id="buffer_content_header">`;
+
+			if (this.display_elem_using_rows)
+			{
+				resource_content += `<td>Member</td><td>Value</td>`;
+				if (this.structure_metadata)
 				{
-					value = (new Int32Array(raw_buffer_data, byte_offset, 1))[0];
+					resource_content += `<td>Hex</td>`;
 				}
-				else if (display == 'uint')
+			}
+			else
+			{
+				resource_content += `<td><input type="text" id="buffer_address_input" onchange="g_view.resource_view.go_to_address();" style="width: 100%;" placeholder="Address..." value="${address_search_value}" /></td>`;
+				for (var i = 0; i < member_count; i++)
 				{
-					value = (new Uint32Array(raw_buffer_data, byte_offset, 1))[0];
+					resource_content += `<td>${member_names[i]}</td>`;
 				}
-				else if (display == 'hex')
+			}
+
+			resource_content += `
+						<td></td>
+					</tr>
+				</table>`;
+
+			document.getElementById('buffer_content_pannel').innerHTML = resource_content;
+		}
+		else
+		{
+			var table_dom = document.getElementById("buffer_content_table");
+			while (table_dom.lastElementChild != table_dom.firstElementChild)
+			{
+				table_dom.removeChild(table_dom.lastElementChild);
+			}
+		}
+
+		{
+			// Start padding
+			var resource_content = ``;
+			{
+				resource_content += `
+					<tr>
+						<td style="height: ${k_buffer_pixel_per_row * display_start_row_id}px; padding: 0;"></td>
+					</tr>`;
+			}
+
+			for (var row_id = display_start_row_id; row_id < display_end_row_id; row_id++)
+			{
+				var elem_id = row_id;
+				if (this.display_elem_using_rows)
 				{
-					value = (new Uint32Array(raw_buffer_data, byte_offset, 1))[0].toString(16);;
+					elem_id = 0;
 				}
-				else if (display == 'ishort')
+
+				var classes = '';
+				if (elem_id == this.selected_address)
 				{
-					value = (new Int16Array(raw_buffer_data, byte_offset, 1))[0];
+					classes = 'highlighted';
 				}
-				else if (display == 'ushort')
+
+				resource_content += `
+					<tr class="${classes}">`;
+
+				if (this.display_elem_using_rows)
 				{
-					value = (new Uint16Array(raw_buffer_data, byte_offset, 1))[0];
+					resource_content += `
+						<td>${member_names[row_id]}</td>`;
 				}
-				else if (display == 'char')
+				else
 				{
-					value = (new Int8Array(raw_buffer_data, byte_offset, 1))[0];
+					resource_content += `
+						<td>0x${(elem_id * byte_per_element).toString(16)}</td>`;
 				}
-				else if (display == 'uchar')
+
+				for (var i = 0; i < column_count; i++)
 				{
-					value = (new Uint8Array(raw_buffer_data, byte_offset, 1))[0];
+					var member_id = i;
+					if (this.display_elem_using_rows)
+					{
+						member_id = row_id;
+					}
+
+					var byte_offset = byte_per_element * elem_id + member_byte_offset[member_id];
+
+					var value = null;
+					var display = member_format[member_id];
+
+					if (this.display_elem_using_rows && i == 1 && this.structure_metadata)
+					{
+						display = `hex(${display})`;
+					}
+
+					if (display == 'float')
+					{
+						value = decode_float32((new Uint32Array(this.raw_buffer_data, byte_offset, 1))[0]);
+					}
+					else if (display == 'half')
+					{
+						value = decode_float16((new Uint16Array(this.raw_buffer_data, byte_offset, 1))[0]);
+					}
+					else if (display == 'int')
+					{
+						value = (new Int32Array(this.raw_buffer_data, byte_offset, 1))[0];
+					}
+					else if (display == 'uint')
+					{
+						value = (new Uint32Array(this.raw_buffer_data, byte_offset, 1))[0];
+					}
+					else if (display == 'short')
+					{
+						value = (new Int16Array(this.raw_buffer_data, byte_offset, 1))[0];
+					}
+					else if (display == 'ushort')
+					{
+						value = (new Uint16Array(this.raw_buffer_data, byte_offset, 1))[0];
+					}
+					else if (display == 'char')
+					{
+						value = (new Int8Array(this.raw_buffer_data, byte_offset, 1))[0];
+					}
+					else if (display == 'uchar')
+					{
+						value = (new Uint8Array(this.raw_buffer_data, byte_offset, 1))[0];
+					}
+					else if (display == 'hex(int)' || display == 'hex(uint)' || display == 'hex(float)')
+					{
+						value = '0x' + (new Uint32Array(this.raw_buffer_data, byte_offset, 1))[0].toString(16).padStart(8, 0);
+					}
+					else if (display == 'hex(short)' || display == 'hex(ushort)' || display == 'hex(half)')
+					{
+						value = '0x' + (new Uint16Array(this.raw_buffer_data, byte_offset, 1))[0].toString(16).padStart(4, 0);
+					}
+					else if (display == 'hex(char)' || display == 'hex(uchar)')
+					{
+						value = '0x' + (new Uint8Array(this.raw_buffer_data, byte_offset, 1))[0].toString(16).padStart(2, 0);
+					}
+					else if (display == 'bin(int)' || display == 'bin(uint)' || display == 'bin(float)')
+					{
+						value = '0b' + (new Uint32Array(this.raw_buffer_data, byte_offset, 1))[0].toString(2).padStart(32, 0);
+					}
+					else if (display == 'bin(short)' || display == 'bin(ushort)' || display == 'bin(half)')
+					{
+						value = '0b' + (new Uint16Array(this.raw_buffer_data, byte_offset, 1))[0].toString(2).padStart(16, 0);
+					}
+					else if (display == 'bin(char)' || display == 'bin(uchar)')
+					{
+						value = '0b' + (new Uint8Array(this.raw_buffer_data, byte_offset, 1))[0].toString(2).padStart(8, 0);
+					}
+					else
+					{
+						value = `Unknown ${display}`;
+					}
+
+					resource_content += `<td>${value}</td>`;
 				}
-				else
+
+				resource_content += `
+						<td></td>
+					</tr>`;
+			} // for (var row_id = display_start_row_id; row_id < display_end_row_id; row_id++)
+
+			if (row_count == 0)
+			{
+				resource_content += `
+					<tr>
+						<td colspan="${column_count + 2}" class="empty" align="center">Empty</td>
+					</tr>`;
+			}
+
+			// End padding
+			{
+				resource_content += `
+					<tr>
+						<td style="height: ${k_buffer_pixel_per_row * (row_count - display_end_row_id)}px; padding: 0;"></td>
+					</tr>`;
+			}
+
+			document.getElementById('buffer_content_table').innerHTML += resource_content;
+		}
+
+		document.getElementById('buffer_viewport').scrollTop = scroll_top;
+
+		{
+			var min_row_count = 10;
+			document.getElementById('buffer_viewport').style.height = `${k_scroll_width + Math.min(Math.max(row_count + 1, min_row_count), k_buffer_max_display_row_count) * k_buffer_pixel_per_row}px`;
+		}
+	}
+
+	release()
+	{
+		this.raw_buffer_data = null;
+	}
+} // class GenericBufferView
+
+
+class BufferView extends GenericBufferView
+{
+	constructor(subresource_version_info, resource_desc)
+	{
+		super(subresource_version_info, resource_desc);
+		this.raw_buffer_data_path = `Resources/${get_subresource_unique_version_name(this.subresource_version_info)}.bin`;
+
+		if (resource_desc['Metadata'] != k_null_json_ptr)
+		{
+			this.structure_metadata = load_structure_metadata(resource_desc['Metadata']);
+		}
+	}
+} // class BufferView
+
+
+// -------------------------------------------------------------------- PARAMETER STRUCTURE
+
+class ParameterStructureView extends GenericBufferView
+{
+	constructor(subresource_version_info, resource_desc, structure_metadata)
+	{
+		super(subresource_version_info, resource_desc);
+		this.raw_buffer_data_path = `Structures/${subresource_version_info['resource']}.bin`;
+		this.structure_metadata = structure_metadata;
+	}
+
+} // class ParameterStructureView
+
+function display_pass_parameters(pass_id)
+{
+	display_pass_internal(pass_id);
+
+	var pass_data = g_passes[pass_id];
+	var structure_metadata = load_structure_metadata(pass_data['ParametersMetadata']);
+
+	var subresource_version_info = {};
+	subresource_version_info['subresource'] = null;
+	subresource_version_info['resource'] = pass_data['Parameters'];
+	subresource_version_info['pass'] = pass_data['Pointer'];
+	subresource_version_info['draw'] = -1;
+
+	var resource_desc = {};
+	resource_desc['Name'] = 'PassParameters';
+	resource_desc['ByteSize'] = structure_metadata['Size'];
+	resource_desc['Desc'] = 'FRDGParameterStruct';
+	resource_desc['Type'] = 'Parameters';
+	resource_desc['BytesPerElement'] = resource_desc['ByteSize'];
+	resource_desc['NumElements'] = 1;
+	resource_desc['Metadata'] = pass_data['ParametersMetadata'];
+	resource_desc['Usage'] = [];
+
+	g_view.set_resource_view(new ParameterStructureView(subresource_version_info, resource_desc, structure_metadata));
+}
+
+
+// -------------------------------------------------------------------- UI
+
+function update_href_selection(parent)
+{
+	var all_a = parent.getElementsByTagName("a");
+
+	var navs = [location.hash.substring(1)];
+	if (g_view)
+	{
+		navs = navs.concat(g_view.navigations);
+	}
+
+	for (let a of all_a)
+	{
+		var href = a.href.split('#');
+
+		if (navs.includes(`${href[1]}`))
+		{
+			a.classList.add('match_location_hash');
+		}
+		else if (a.classList.contains('match_location_hash'))
+		{
+			a.classList.remove('match_location_hash');
+		}
+	};
+}
+
+function update_value_selection(parent, selected_value)
+{
+	var all_a = parent.getElementsByTagName("input");
+
+	for (let a of all_a)
+	{
+		if (a.value == selected_value)
+		{
+			a.classList.add('match_value');
+		}
+		else if (a.classList.contains('match_value'))
+		{
+			a.classList.remove('match_value');
+		}
+	};
+}
+
+function render_selection_list_html(parent_dom, display_list, options)
+{
+	if ('deduplicate' in options)
+	{
+		var href_set = new Set();
+		var new_display_list = [];
+		for (const display_list_item of display_list)
+		{
+			if (display_list_item['href'])
+			{
+				if (!href_set.has(display_list_item['href']))
 				{
-					value = `Unknown ${display}`;
+					href_set.add(display_list_item['href']);
+					new_display_list.push(display_list_item);
 				}
+			}
+			else
+			{
+				new_display_list.push(display_list_item);
+			}
+		}
+		display_list = new_display_list;
+	}
 
-				resource_content += `<td>${value}</td>`;
+	if ('sort' in options)
+	{
+		display_list.sort(function(a, b)
+		{
+			if (a['name'] < b['name'])
+			{
+				return -1;
+			}
+			else if (a['name'] > b['name'])
+			{
+				return 1;
 			}
+			return 0;
+		});
+	}
+
+	var search = '';
+	if ('search' in options)
+	{
+		search = options['search'];
+	}
+
+	var html = '';
+	for (const display_list_item of display_list)
+	{
+		if (search && !display_list_item['name'].toLowerCase().includes(search.toLowerCase()))
+		{
+			continue;
+		}
 
-			resource_content += `</tr>`;
+		if (display_list_item['href'])
+		{
+			html += `<a href="${display_list_item['href']}">${display_list_item['name']}</a>`;
+		}
+		else
+		{
+			html += `<a>${display_list_item['name']}</a>`;
 		}
+	}
 
-		resource_content += `</table>`;
+	if (html == '')
+	{
+		if (search)
+		{
+			html += `<div class="empty">No matches for "${search}"</div>`;
+		}
+		else
+		{
+			html += `<div class="empty">None</div>`;
+		}
 	}
 
-	document.getElementById('buffer_content_pannel').innerHTML = resource_content;
+	parent_dom.innerHTML = html;
+	update_href_selection(parent_dom);
 }
 
 
-
 // -------------------------------------------------------------------- WINDOW LAYOUT
 
 function init_top()
 {
-	console.log(g_infos);
 	document.getElementById('top_bar').innerHTML = `${g_infos['Project']} - ${g_infos['Platform']} - ${g_infos['RHI']} - ${g_infos['BuildVersion']} - ${g_infos['DumpTime']} `;
 }
 
+function init_top_buttons()
+{
+	document.getElementById('top_buttons').innerHTML = `
+		<a href="#display_infos();">Infos</a>
+		<a href="#display_cvars();">CVars</a>
+		<a href="#display_log();">Log</a>`;
+}
+
 function init_page()
 {
-	g_infos = load_json('Infos.json');
-	g_passes = load_json_dict_sequence('Passes.json');
-	console.log(g_passes[0]);
+	// Load the base information of the page.
+	g_infos = load_json('Base/Infos.json');
+	if (!g_infos)
+	{
+		document.body.innerHTML = 'Please use OpenGPUDumpViewer script for the viewer to work correctly.';
+		return;
+	}
+
+	// Load the cvars used for dumping.
+	{
+		g_dump_cvars = {};
+
+		var cvars_csv = load_text_file('Base/ConsoleVariables.csv');
+		var csv_lines = cvars_csv.split('\n');
+		for (var i = 1; i < csv_lines.length - 1; i++)
+		{
+			var csv_line = csv_lines[i].split(',');
+			var entry = {};
+			entry['name'] = csv_line[0];
+			if (entry['name'].startsWith('r.DumpGPU.'))
+			{
+				entry['type'] = csv_line[1];
+				entry['set_by'] = csv_line[2];
+				entry['value'] = csv_line[3];
+				g_dump_cvars[entry['name']] = entry['value'];
+			}
+		}
+
+	}
+
+	// Load all the passes.
+	g_passes = load_json_dict_sequence('Base/Passes.json');
+	{
+		for (var pass_id = 0; pass_id < g_passes.length; pass_id++)
+		{
+			g_passes[pass_id]['DrawCount'] = 0;
+			if (!g_passes[pass_id]['EventName'])
+			{
+				g_passes[pass_id]['EventName'] = 'Unnamed pass';
+			}
+		}
+	}
+
+	var pass_draw_counts = load_json_dict_sequence('Base/PassDrawCounts.json');
+	if (pass_draw_counts)
+	{
+		pass_draw_counts.forEach(function(pass_draw_count_data) {
+			for (var pass_id = 0; pass_id < g_passes.length; pass_id++)
+			{
+				if (g_passes[pass_id]['Pointer'] == pass_draw_count_data['Pointer'])
+				{
+					g_passes[pass_id]['DrawCount'] = pass_draw_count_data['DrawCount'];
+					break;
+				}
+			}
+		});
+	}
+
 	init_top();
+	init_top_buttons();
 	display_pass_hierarchy();
+
+	// Find the last passes that have a displayable texture 2D.
+	if (location.hash)
+	{
+		return navigate_to_hash();
+	}
+	
+	// Find the last pass that modify the RDG resource set by 'r.DumpGPU.Viewer.Visualize'
+	if ('r.DumpGPU.Viewer.Visualize' in g_dump_cvars && g_dump_cvars['r.DumpGPU.Viewer.Visualize'] != '')
+	{
+		var pass_id_to_open = -1;
+		var subresource_unique_name_to_open = null;
+
+		document.getElementById('resource_search_input').value = g_dump_cvars['r.DumpGPU.Viewer.Visualize'];
+		display_pass_hierarchy();
+
+		for (var pass_id = 0; pass_id < g_passes.length && pass_id_to_open == -1; pass_id++)
+		{
+			var pass_data = g_passes[g_passes.length - 1 - pass_id];
+
+			for (var subresource_unique_name of pass_data['OutputResources'])
+			{
+				var subresource_info = parse_subresource_unique_name(subresource_unique_name);
+				var resource_desc = load_resource_desc(subresource_info['resource']);
+
+				if (resource_desc === null)
+				{
+					continue;
+				}
+
+				if (resource_desc['Name'] == g_dump_cvars['r.DumpGPU.Viewer.Visualize'] && pass_id_to_open == -1)
+				{
+					pass_id_to_open = g_passes.length - 1 - pass_id;
+					subresource_unique_name_to_open = subresource_unique_name;
+					break;
+				}
+			}
+		}
+
+		if (pass_id_to_open != -1)
+		{
+			return redirect_to_hash(`display_output_resource(${pass_id_to_open},'${subresource_unique_name_to_open}');`);
+		}
+	}
+	
+	return display_tip();
 }
 
 function onresize_body()
@@ -1091,12 +3892,63 @@ function onresize_body()
 
 	{
 		var search_h = (30 + 2 * 5) * 2;
+		document.getElementById('pass_hierarchy').style.width = `${Math.floor(w * 0.25) - 10 * 2}px`;
 		document.getElementById('pass_hierarchy').style.height = `${h - top_h - search_h - 10 * 2}px`;
 	}
 
 	{
+		//document.getElementById('main_right_pannel').style.width = document.getElementById('onresize_body_right').style.width;
 		document.getElementById('main_right_pannel').style.height = `${h - top_h - 10 * 2}px`;
 	}
+
+	if (g_view !== null)
+	{
+		var ctx = {};
+		ctx.width = w - Math.floor(w * 0.25);
+		g_view.resize(ctx);
+	}
+}
+
+
+// -------------------------------------------------------------------- NAVIGATION
+
+var g_previous_location_hash = null;
+
+function redirect_to_hash(new_hash)
+{
+	if (new_hash === g_previous_location_hash)
+	{
+		return;
+	}
+
+	location.hash = new_hash;
+}
+
+function navigate_to_hash()
+{
+	if (location.hash === g_previous_location_hash)
+	{
+		return;
+	}
+
+	g_previous_location_hash = location.hash;
+
+	if (location.hash)
+	{
+		var js = location.hash.substring(1);
+		eval(js);
+	}
+	else
+	{
+		display_tip();
+	}
+
+	update_href_selection(document);
+}
+
+function get_current_navigation()
+{
+	return g_previous_location_hash;
 }
 
 </script>
@@ -1104,41 +3956,210 @@ function onresize_body()
 
 <style type="text/css">
 
+
+/* ----------------------------------------------------- them colors */
+
+:root
+{
+	--border-radius: 2px;
+
+	--body-bg-color: rgb(21, 21, 21);
+	--body-txt-size: 12px;
+
+	--div-bg-color: rgb(36, 36, 36);
+	--div-txt-color: rgb(192, 192, 192);
+	--a-bg-color-odd: rgb(26, 26, 26);
+	--a-bg-color-hover: #332222;
+	--a-bg-color-hoverhref: rgb(38, 67, 81);
+	--a-txt-color-disabled: rgb(113, 113, 113);
+
+	--a-bg-color-selected: rgb(38, 187, 255);
+	--a-txt-color-selected: rgb(15, 15, 15);
+
+	--channel-red-color: rgb(255, 38, 38);
+	--channel-green-color: rgb(38, 255, 38);
+	--channel-blue-color: rgb(38, 187, 255);
+	--channel-alpha-color: white;
+
+	--error-txt-color: var(--channel-red-color);
+
+	--header-bg-color: rgb(47, 47, 47);
+
+	--button-border-radius: var(--border-radius);
+	--button-bg-color: rgb(56, 56, 56);
+	--button-bg-color-hover: rgb(87, 87, 87);
+	--button-bg-color-selected: rgb(15, 15, 15);
+	--button-txt-color: var(--div-txt-color);
+	--button-txt-color-hover: rgb(255, 255, 255);
+	--button-txt-color-selected: rgb(38, 187, 255);
+
+
+	--scroll-bg-color: var(--a-bg-color-odd);
+	--scroll-color: rgb(87, 87, 87);
+	--scroll-color-hover: rgb(128, 128, 128);
+
+	--input-bg-color: rgb(15, 15, 15);
+	--input-border: 1px solid rgb(55, 55, 55);
+	--input-border-hover: rgb(83, 83, 83);
+	--input-border-focus: rgb(38, 176, 239);
+}
+
+
+/* ----------------------------------------------------- override default */
+
 body
 {
-	background: #111111;
-	color: #DDDDDD;
+	overflow: hidden;
+	background: var(--body-bg-color);
+	color: var(--div-txt-color);
 	margin: 0;
 	padding: 0;
 	/*font-family: Arial, Helvetica, sans-serif;*/
 	font-family: consolas, "Liberation Mono", courier, monospace;
+	font-size: var(--body-txt-size);
 	/* overflow: hidden; */
 }
 
-table, tr, td, div
+table, tr, td, div, a
 {
 	margin: 0;
 	padding: 0;
 	border-spacing: 0;
 	border-collapse: collapse;
 	vertical-align: top;
+	cursor: default;
+	font-size: inherit;
 }
 
 td
 {
-	display: inline-block;
+	cursor: default;
 }
 
-div
+div, a
 {
 	display: block;
+	cursor: default;
+}
+
+a, a:hover, a:visited
+{
+	color: inherit;
+	font: inherit;
+	text-decoration: none;
+	cursor: default;
+}
+
+
+/* ----------------------------------------------------- external link */
+
+a.external_link
+{
+	color: inherit;
+	cursor: pointer;
+}
+
+a.external_link:hover
+{
+	color: var(--button-txt-color-hover);
+}
+
+a.external_link:active
+{
+	color: var(--button-txt-color-selected);
 }
 
+
+/* ----------------------------------------------------- inputs scroll bars */
+
+input
+{
+	padding: 3px 5px;
+	border-radius: var(--border-radius);
+}
+
+input[type=search]
+{
+	padding: 3px 13px;
+	border-radius: 20px;
+}
+
+textarea
+{
+	padding: 3px 3px;
+}
+
+input, textarea
+{
+	background-color: var(--input-bg-color);
+	color: var(--div-txt-color);
+	border: var(--input-border);
+	outline:  none;
+	font-size: inherit;
+	font: inherit;
+}
+
+input[readonly], textarea[readonly]
+{
+	color: var(--a-txt-color-disabled);
+}
+
+input:not([readonly]):hover,
+textarea:not([readonly]):hover
+{
+	border-color: var(--input-border-hover);
+}
+
+input[type=text]:not([readonly]):focus,
+input[type=search]:not([readonly]):focus,
+textarea:not([readonly]):focus
+{
+	color: rgb(254, 254, 254);
+	border-color: var(--input-border-focus);
+}
+
+input:focus, textarea:focus
+{
+	outline: none;
+}
+
+input::placeholder
+{
+	color: rgb(76, 76, 76);
+}
+
+
+/* ----------------------------------------------------- webkit scroll bars */
+
+*::-webkit-scrollbar {
+	width: 8px;
+	height: 8px;
+	background: var(--scroll-bg-color);
+}
+
+*::-webkit-scrollbar-corner, *::-webkit-scrollbar-track {
+	background: var(--scroll-bg-color);
+}
+
+*::-webkit-scrollbar-thumb {
+	background-color: var(--scroll-color);
+	border-radius: 20px;
+}
+
+*::-webkit-scrollbar-thumb:hover {
+	background-color: var(--scroll-color-hover);
+	border-radius: 20px;
+}
+
+
+/* ----------------------------------------------------- common layout */
+
 .main_div
 {
-	background: #222222;
+	background: var(--div-bg-color);
 	padding: 5px;
 	margin: 5px;
+	border-radius: 3px;
 }
 
 #main_right_pannel .main_div
@@ -1147,7 +4168,26 @@ div
 }
 
 
-/* ----------------------------------------------------- common layout */
+/* ----------------------------------------------------- Selection list */
+
+.selection_list_title
+{
+	font-size: 20;
+	font-weight: bold;
+	padding: 5px 20px 10px 20px;
+}
+
+.selection_list_search
+{
+	padding: 0px 5px;
+	height: 30px;
+	width: auto;
+}
+
+.selection_list_search input[type=search]
+{
+	width: 100%;
+}
 
 .selection_list
 {
@@ -1156,75 +4196,181 @@ div
 	overflow-y: scroll;
 }
 
-.selection_list div
+.selection_list a
 {
 	width: auto;
-	padding: 5px 10px;
+	padding: 5px 20px;
 	white-space: nowrap;
 }
 
-.selection_list div:nth-child(odd)
+.selection_list a:nth-child(odd)
 {
-	background: #171717;
+	background: var(--a-bg-color-odd);
 }
 
-.selection_list div.selected
+.selection_list a:not(.match_location_hash):hover
 {
-	background: #114422;
+	background: var(--a-bg-color-hover);
 }
 
-.selection_list div:hover
+.selection_list a:not(.match_location_hash)[href]:hover
 {
-	background: #112244;
+	background: var(--a-bg-color-hoverhref);
+	color: var(--button-txt-color-hover);
 	cursor: pointer;
 }
 
-.selection_list div.disabled
+/*.selection_list a.match_location_hash
+{
+	background: var(--a-bg-color-selected);
+	color: var(--a-txt-color-selected);
+}*/
+
+.selection_list a.match_location_hash
+{
+	background: var(--button-bg-color-selected);
+	color: var(--button-txt-color-selected);
+}
+
+.selection_list a.disabled
 {
-	color: #999999;
+	color: var(--a-txt-color-disabled);
+}
+
+.selection_list div.empty
+{
+	width: 100%;
+	margin: 20px 0 0 0;
+	color: var(--a-txt-color-disabled);
+	text-align: center;
 }
 
 #main_right_pannel .selection_list
 {
-	max-height: 20%;
+	min-height: 100px;
+	max-height: 200px;
 }
 
-.selection_list_title
+
+/* ----------------------------------------------------- table */
+
+.pretty_table tr.header
 {
-	font-size: 20;
+	background: var(--header-bg-color);
 	font-weight: bold;
-	padding: 5px 20px 10px 20px;
 }
 
-.resource_desc tr:nth-child(odd)
+.pretty_table tr:not(.header):nth-child(odd) td
 {
-	background: #171717;
+	background: var(--a-bg-color-odd);
 }
 
-.resource_desc tr td
+.pretty_table tr:not(.header):hover td:not(.empty)
 {
+	background: var(--a-bg-color-hover);
+}
+
+.pretty_table tr td
+{
+	display: table-cell;
 	padding: 5px 20px;
 }
 
-.resource_desc tr td
+.pretty_table tr td:first-child
 {
-	text-align: right;
+	width: 150px;
+}
+
+.pretty_table tr.header td
+{
+	padding: 5px 30px;
+}
+
+.pretty_table tr.header td:not(:first-child)
+{
+	border-left: 1px solid rgb(36, 36, 36);
 }
 
-.resource_desc tr td:first-child
+.pretty_table .empty
+{
+	padding: 20px 0 0 0;
+	color: var(--a-txt-color-disabled);
+	text-align: center;
+}
+
+
+.resource_desc tr td:not(.empty)
 {
 	width: 150px;
+	text-align: right;
+}
+
+
+/* ----------------------------------------------------- button */
+
+.button_list a,
+.button_list input
+{
+	margin: 0;
+	display: inline-block;
+	display: table-cell;
+	padding: 3px 10px;
+	background-color: var(--button-bg-color);
+	color: var(--button-txt-color);
+	cursor: pointer;
+	border-radius: 0;
+}
+
+.button_list a:first-child,
+.button_list input:first-child
+{
+	border-top-left-radius: var(--button-border-radius);
+	border-bottom-left-radius: var(--button-border-radius);
+}
+
+.button_list a:last-child,
+.button_list input:last-child
+{
+	border-top-right-radius: var(--button-border-radius);
+	border-bottom-right-radius: var(--button-border-radius);
+}
+
+.button_list a:not(.match_location_hash):not(:active):hover,
+.button_list input:not(.match_value):not(:active):hover
+{
+	background-color: var(--button-bg-color-hover);
+	color: var(--button-txt-color-hover);
+}
+
+.button_list a:active,
+.button_list a.match_location_hash,
+.button_list input:active,
+.button_list input.match_value
+{
+	background-color: var(--button-bg-color-selected);
+	color: var(--button-txt-color-selected);
 }
 
 
+/* ----------------------------------------------------- top bar */
 
 #top_bar
 {
 	font-size: 20;
 	font-weight: bold;
 	padding: 10px;
+	display: inline-block;
+}
+
+#top_buttons
+{
+	padding: 12px;
+	display: inline-block;
 }
 
+
+/* ----------------------------------------------------- main_right_pannel */
+
 #main_right_pannel
 {
 	width: auto;
@@ -1239,6 +4385,91 @@ div
 	padding: 10px 40px;
 }
 
+#main_right_pannel .pass_title .button_list
+{
+	font-size: var(--body-txt-size);
+	margin-top: 3px;
+}
+
+#cvars_pannel
+{
+	height: auto;
+	max-height: none;
+}
+
+#cvars_pannel .pretty_table tr td:nth-child(2)
+{
+	text-align: right;
+}
+
+
+/* ----------------------------------------------------- TextureView */
+
+#canvas_outter
+{
+	background-color: var(--input-bg-color);
+	color: var(--div-txt-color);
+	border: var(--input-border);
+}
+
+#canvas_outter:not(.selected):hover
+{
+	border-color: var(--input-border-hover);
+}
+
+#canvas_outter.selected
+{
+	border-color: var(--input-border-focus);
+}
+
+#canvas_outter #canvas_header
+{
+	padding: 3px;
+	border-bottom: var(--input-border);
+}
+
+#canvas_outter #canvas_header .button_list
+{
+	display: inline-block;
+}
+
+#canvas_outter #canvas_footer
+{
+	padding: 3px 10px;
+	border-top: var(--input-border);
+}
+
+#canvas_outter #canvas_viewport
+{
+	overflow: hidden;
+    background-image:
+		linear-gradient(45deg, #000 25%, transparent 25%),
+		linear-gradient(45deg, transparent 75%, #000 75%),
+		linear-gradient(45deg, transparent 75%, #000 75%),
+		linear-gradient(45deg, #000 25%, var(--input-bg-color) 25%);
+    background-size:16px 16px;
+    background-position:0 0, 0 0, -8px -8px, 8px 8px;
+}
+
+#canvas_outter .error_msg
+{
+	width: 100%;
+	height: 100%;
+	text-align: center;
+	color: var(--error-txt-color);
+}
+
+#canvas_outter canvas
+{
+	cursor: crosshair;
+	image-rendering: optimizeSpeed;
+	image-rendering: -moz-crisp-edges;
+	image-rendering: -webkit-optimize-contrast;
+	image-rendering: -o-crisp-edges;
+	image-rendering: pixelated;
+	-ms-interpolation-mode: nearest-neighbor;
+}
+
 #texture_visualization_code_input
 {
 	min-width: 800px;
@@ -1256,44 +4487,95 @@ div
 	width: auto;
 }
 
-#buffer_content_table tr:nth-child(even)
+
+/* ----------------------------------------------------- BufferView */
+
+#buffer_outter
+{
+	background-color: var(--input-bg-color);
+	color: var(--div-txt-color);
+	border: var(--input-border);
+}
+
+#buffer_outter:not(.selected):hover
 {
-	background: #171717;
+	border-color: var(--input-border-hover);
 }
 
-#buffer_content_table tr td
+#buffer_outter.selected
 {
-	padding: 3px 5px;
-	width: 40px;
+	border-color: var(--input-border-focus);
+}
+
+#buffer_outter #buffer_viewport_header
+{
+	padding: 3px;
+	border-bottom: var(--input-border);
+}
+
+#buffer_outter #buffer_viewport
+{
+	overflow: hidden;
+}
+
+#buffer_outter #buffer_viewport #buffer_content_table tr.header
+{
+	position: sticky;
+	top: 0;
+}
+
+#buffer_outter #buffer_viewport #buffer_content_table:not(.display_elem_using_rows) tr.header td:first-child
+{
+	padding: 0 3px;
+}
+
+#buffer_outter #buffer_viewport #buffer_content_table tr td:not(.empty)
+{
+	padding: 3px 5px 0px 5px;
+	width: 60px;
 	height: 20px;
 	overflow: hidden;
-	display: table-cell;
 	text-align: right;
-	font-family: consolas, "Liberation Mono", courier, monospace;
+	/*font-family: consolas, "Liberation Mono", courier, monospace;*/
+}
+
+#buffer_outter #buffer_viewport #buffer_content_table.display_elem_using_rows tr td:not(.empty):first-child
+{
+	text-align: left;
+}
+
+#buffer_outter #buffer_viewport #buffer_content_table tr td:first-child
+{
+	width: 100px;
+}
+
+#buffer_outter #buffer_viewport #buffer_content_table tr td:last-child
+{
+	width: auto;
 }
 
-canvas
+#buffer_outter #buffer_viewport #buffer_content_table tr.highlighted
 {
-	background-color: black;
-	border: 1px solid #333333;
+	color: var(--button-txt-color-selected);
 }
 
 </style>
 
 
-<body onload="onresize_body(); init_page();" onresize="onresize_body();">
+<body onload="onresize_body(); init_page();" onresize="onresize_body();" onhashchange="navigate_to_hash();">
 	<table cellpadding="0" cellspacing="0" border="0" id="onresize_body_table">
 		<tr>
 			<td colspan="2" id="onresize_body_top_bar">
 				<div id="top_bar"></div>
+				<div id="top_buttons" class="button_list"></div>
 			</td>
 		</tr>
 		<tr>
 			<td valign="top" id="onresize_body_left">
 				<div class="main_div">
 					<div id="pass_hierarchy_search" style="padding: 5px; height: 70px;">
-						<input type="text" id="pass_search_input" oninput="display_pass_hierarchy();" onchange="display_pass_hierarchy();" style="width: 100%;" placeholder="Search pass..." />
-						<input type="text" id="resource_search_input" oninput="display_pass_hierarchy();" onchange="display_pass_hierarchy();" style="width: 100%; margin-top: 10px;" placeholder="Search resource..." />
+						<input type="search" id="pass_search_input" oninput="display_pass_hierarchy();" onchange="display_pass_hierarchy();" style="width: 100%;" placeholder="Search pass..." />
+						<input type="search" id="resource_search_input" oninput="display_pass_hierarchy();" onchange="display_pass_hierarchy();" style="width: 100%; margin-top: 10px;" placeholder="Search resource..." />
 					</div>
 					<div id="pass_hierarchy" class="selection_list"></div>
 				</div>
diff --git a/Engine/Shaders/Private/Tools/DumpTexture.usf b/Engine/Shaders/Private/Tools/DumpTexture.usf
new file mode 100644
--- /dev/null
+++ b/Engine/Shaders/Private/Tools/DumpTexture.usf
@@ -0,0 +1,44 @@
+// Copyright Epic Games, Inc. All Rights Reserved.
+
+#include "../Common.ush"
+
+
+//------------------------------------------------------- CONFIG
+
+#define TILE_SIZE 8
+
+#if DIM_TEXTURE_TYPE == 0 // Texture2DFloatNoMSAA
+	#define payload float4	
+
+#elif DIM_TEXTURE_TYPE == 1 // Texture2DIntNoMSAA
+	#define payload uint4
+
+#elif DIM_TEXTURE_TYPE == 2 // Texture2DDepthStencilNoMSAA
+	#define payload uint4
+	#define SWIZZLE_STENCIL 1
+
+#else
+	#error unknown DIM_TEXTURE_TYPE
+
+#endif
+
+
+//------------------------------------------------------- PARAMETERS
+
+Texture2D<payload> Texture;
+RWTexture2D<payload> StagingOutput;
+
+
+//------------------------------------------------------- ENTRY POINT
+
+[numthreads(TILE_SIZE, TILE_SIZE, 1)]
+void MainCS(uint2 DispatchThreadId : SV_DispatchThreadID)
+{
+	payload PixelValue = Texture[DispatchThreadId];
+	
+	#if defined(SWIZZLE_STENCIL)
+		PixelValue = payload(PixelValue STENCIL_COMPONENT_SWIZZLE, 0, 0, 0);
+	#endif
+
+	StagingOutput[DispatchThreadId] = PixelValue; 
+}
diff --git a/Engine/Source/Editor/LevelEditor/Private/LevelEditor.cpp b/Engine/Source/Editor/LevelEditor/Private/LevelEditor.cpp
--- a/Engine/Source/Editor/LevelEditor/Private/LevelEditor.cpp
+++ b/Engine/Source/Editor/LevelEditor/Private/LevelEditor.cpp
@@ -1696,6 +1696,11 @@ void FLevelEditorModule::BindGlobalLevelEditorCommands()
 		Commands.ProfileGPU,
 		FExecuteAction::CreateStatic( &FLevelEditorActionCallbacks::ExecuteExecCommand, FString( TEXT("PROFILEGPU") ) )
 		);
+	
+	ActionList.MapAction(
+		Commands.DumpGPU,
+		FExecuteAction::CreateStatic( &FLevelEditorActionCallbacks::ExecuteExecCommand, FString( TEXT("DUMPGPU") ) )
+		);
 
 	ActionList.MapAction(
 		Commands.ResetAllParticleSystems,
diff --git a/Engine/Source/Editor/LevelEditor/Private/LevelEditorActions.cpp b/Engine/Source/Editor/LevelEditor/Private/LevelEditorActions.cpp
--- a/Engine/Source/Editor/LevelEditor/Private/LevelEditorActions.cpp
+++ b/Engine/Source/Editor/LevelEditor/Private/LevelEditorActions.cpp
@@ -3619,6 +3619,7 @@ void FLevelEditorCommands::RegisterCommands()
 
 	UI_COMMAND( RecompileShaders, "Recompile Changed Shaders", "Recompiles shaders which are out of date", EUserInterfaceActionType::Button, FInputChord( EModifierKey::Shift|EModifierKey::Control, EKeys::Period ) );
 	UI_COMMAND( ProfileGPU, "Profile GPU", "Profiles the GPU for the next frame and opens a window with profiled data", EUserInterfaceActionType::Button, FInputChord( EModifierKey::Shift|EModifierKey::Control, EKeys::Comma ) );
+	UI_COMMAND( DumpGPU, "Dump GPU", "Dump the GPU intermediary resources for the next frame and opens explorer", EUserInterfaceActionType::Button, FInputChord( EModifierKey::Shift|EModifierKey::Control, EKeys::Slash ) );
 
 	UI_COMMAND( ResetAllParticleSystems, "Reset All Particle Systems", "Resets all particle system emitters (removes all active particles and restarts them)", EUserInterfaceActionType::Button, FInputChord( EModifierKey::Shift, EKeys::Slash ) );
 	UI_COMMAND( ResetSelectedParticleSystem, "Resets Selected Particle Systems" , "Resets selected particle system emitters (removes all active particles and restarts them)", EUserInterfaceActionType::Button, FInputChord( EKeys::Slash ) );
diff --git a/Engine/Source/Editor/LevelEditor/Public/LevelEditorActions.h b/Engine/Source/Editor/LevelEditor/Public/LevelEditorActions.h
--- a/Engine/Source/Editor/LevelEditor/Public/LevelEditorActions.h
+++ b/Engine/Source/Editor/LevelEditor/Public/LevelEditorActions.h
@@ -614,6 +614,7 @@ public:
 	TSharedPtr< FUICommandInfo > ShowSelectedDetails;
 	TSharedPtr< FUICommandInfo > RecompileShaders;
 	TSharedPtr< FUICommandInfo > ProfileGPU;
+	TSharedPtr< FUICommandInfo > DumpGPU;
 
 	TSharedPtr< FUICommandInfo > ResetAllParticleSystems;
 	TSharedPtr< FUICommandInfo > ResetSelectedParticleSystem;
diff --git a/Engine/Source/Runtime/Core/Private/HAL/ConsoleManager.cpp b/Engine/Source/Runtime/Core/Private/HAL/ConsoleManager.cpp
--- a/Engine/Source/Runtime/Core/Private/HAL/ConsoleManager.cpp
+++ b/Engine/Source/Runtime/Core/Private/HAL/ConsoleManager.cpp
@@ -2010,6 +2010,9 @@ void CreateConsoleVariables()
 	IConsoleManager::Get().RegisterConsoleCommand(TEXT("r.ResetRenderTargetsExtent"), TEXT("To reset internal render target extents"), ECVF_Cheat);
 #endif // !(UE_BUILD_SHIPPING || UE_BUILD_TEST)
 
+#if WITH_PROFILEGPU
+	IConsoleManager::Get().RegisterConsoleCommand(TEXT("DumpGPU"), TEXT("Dump one frame of rendering intermediary resources to disk."), ECVF_Cheat);
+#endif
 
 #if	!UE_BUILD_SHIPPING
 	IConsoleManager::Get().RegisterConsoleCommand( TEXT( "DumpConsoleCommands" ), TEXT( "Dumps all console vaiables and commands and all exec that can be discovered to the log/console" ), ECVF_Default );
diff --git a/Engine/Source/Runtime/Core/Public/Misc/Build.h b/Engine/Source/Runtime/Core/Public/Misc/Build.h
--- a/Engine/Source/Runtime/Core/Public/Misc/Build.h
+++ b/Engine/Source/Runtime/Core/Public/Misc/Build.h
@@ -382,6 +382,17 @@
 // draw events with "TOGGLEDRAWEVENTS" "r.ShowMaterialDrawEvents" (for ProfileGPU, Pix, Razor, RenderDoc, ...) and the "ProfileGPU" command are normally compiled out for TEST and SHIPPING
 #define WITH_PROFILEGPU (!(UE_BUILD_SHIPPING || UE_BUILD_TEST) || (UE_BUILD_TEST && ALLOW_PROFILEGPU_IN_TEST) || (UE_BUILD_SHIPPING && ALLOW_PROFILEGPU_IN_SHIPPING))
 
+#ifndef ALLOW_DUMPGPU_IN_TEST
+	#define ALLOW_DUMPGPU_IN_TEST 1
+#endif
+
+#ifndef ALLOW_DUMPGPU_IN_SHIPPING
+	#define ALLOW_DUMPGPU_IN_SHIPPING 0
+#endif
+
+// DumpGPU command
+#define WITH_DUMPGPU (!(UE_BUILD_SHIPPING || UE_BUILD_TEST) || (UE_BUILD_TEST && ALLOW_DUMPGPU_IN_TEST) || (UE_BUILD_SHIPPING && ALLOW_DUMPGPU_IN_SHIPPING))
+
 #ifndef ALLOW_CHEAT_CVARS_IN_TEST
 	#define ALLOW_CHEAT_CVARS_IN_TEST 1
 #endif
diff --git a/Engine/Source/Runtime/Engine/Classes/Engine/Engine.h b/Engine/Source/Runtime/Engine/Classes/Engine/Engine.h
--- a/Engine/Source/Runtime/Engine/Classes/Engine/Engine.h
+++ b/Engine/Source/Runtime/Engine/Classes/Engine/Engine.h
@@ -2027,6 +2027,10 @@ public:
 	bool HandleProfileGPUCommand( const TCHAR* Cmd, FOutputDevice& Ar );	
 #endif
 
+#if WITH_DUMPGPU
+	bool HandleDumpGPUCommand( const TCHAR* Cmd, FOutputDevice& Ar );
+#endif
+
 	// Compile in Debug or Development
 #if !(UE_BUILD_SHIPPING || UE_BUILD_TEST) && WITH_HOT_RELOAD
 	bool HandleHotReloadCommand( const TCHAR* Cmd, FOutputDevice& Ar );
diff --git a/Engine/Source/Runtime/Engine/Private/UnrealEngine.cpp b/Engine/Source/Runtime/Engine/Private/UnrealEngine.cpp
--- a/Engine/Source/Runtime/Engine/Private/UnrealEngine.cpp
+++ b/Engine/Source/Runtime/Engine/Private/UnrealEngine.cpp
@@ -255,6 +255,19 @@ UnrealEngine.cpp: Implements the UEngine class and helpers.
 #include "Engine/InstancedStaticMesh.h"
 #include "IDeviceProfileSelectorModule.h"
 
+#if WITH_DUMPGPU
+	#include "RenderGraph.h"
+#endif
+
+#if WITH_DUMPGPU
+
+static TAutoConsoleVariable<int32> GDumpGPUScreenshot(
+	TEXT("r.DumpGPU.Screenshot"), 1,
+	TEXT("Whether to take a final screenshot."),
+	ECVF_Default);
+
+#endif
+
 DEFINE_LOG_CATEGORY(LogEngine);
 IMPLEMENT_MODULE( FEngineModule, Engine );
 
@@ -4528,6 +4541,13 @@ bool UEngine::Exec( UWorld* InWorld, const TCHAR* Cmd, FOutputDevice& Ar )
 	}	
 #endif // #if !UE_BUILD_SHIPPING
 
+#if WITH_DUMPGPU
+	else if (FParse::Command(&Cmd, TEXT("DUMPGPU")))
+	{
+		return HandleDumpGPUCommand(Cmd, Ar);
+	}
+#endif
+
 #if	!(UE_BUILD_SHIPPING || UE_BUILD_TEST) && WITH_HOT_RELOAD
 	else if( FParse::Command(&Cmd,TEXT("HotReload")) )
 	{
@@ -5473,6 +5493,33 @@ bool UEngine::HandleProfileGPUCommand( const TCHAR* Cmd, FOutputDevice& Ar )
 
 #endif // WITH_PROFILEGPU
 
+#if WITH_DUMPGPU
+bool UEngine::HandleDumpGPUCommand(const TCHAR* Cmd, FOutputDevice& Ar)
+{
+	if (!FApp::CanEverRender())
+	{
+		return true;
+	}
+
+	check(IsInGameThread());
+	TArray<FString> Args;
+	FString ResourceDumpDirectory = FRDGBuilder::BeginResourceDump(Args);
+
+	if (!ResourceDumpDirectory.IsEmpty())
+	{
+		// Dump a screenshot
+		if (GDumpGPUScreenshot.GetValueOnGameThread())
+		{
+			FString FileName = ResourceDumpDirectory / TEXT("Base/Screenshot.png");
+			FScreenshotRequest::RequestScreenshot(FileName, /* bShowUI = */ true, /* bAddFilenameSuffix = */ false);
+		}
+	}
+
+	return true;
+}
+
+#endif // WITH_DUMPGPU
+
 #if !UE_BUILD_SHIPPING
 bool UEngine::HandleStartFPSChartCommand( const TCHAR* Cmd, FOutputDevice& Ar )
 {
diff --git a/Engine/Source/Runtime/Launch/Private/LaunchEngineLoop.cpp b/Engine/Source/Runtime/Launch/Private/LaunchEngineLoop.cpp
--- a/Engine/Source/Runtime/Launch/Private/LaunchEngineLoop.cpp
+++ b/Engine/Source/Runtime/Launch/Private/LaunchEngineLoop.cpp
@@ -5517,14 +5517,8 @@ void FEngineLoop::Tick()
 		FCoreDelegates::OnEndFrame.Broadcast();
 
 		// end of RDG resource dump
-		#if WITH_ENGINE
-			#if RDG_DUMP_RESOURCES
-				ENQUEUE_RENDER_COMMAND(EndRDGResourceDump)(
-					[](FRHICommandListImmediate& RHICmdList)
-					{
-						FRDGBuilder::EndResourceDump();
-					});
-			#endif
+		#if WITH_ENGINE && WITH_DUMPGPU
+			FRDGBuilder::EndResourceDump();
 		#endif
 
 		#if !UE_SERVER && WITH_ENGINE
diff --git a/Engine/Source/Runtime/RHI/Public/RHIDefinitions.h b/Engine/Source/Runtime/RHI/Public/RHIDefinitions.h
--- a/Engine/Source/Runtime/RHI/Public/RHIDefinitions.h
+++ b/Engine/Source/Runtime/RHI/Public/RHIDefinitions.h
@@ -2279,6 +2279,181 @@ inline const TCHAR* GetTextureDimensionString(ETextureDimension Dimension)
 	return TEXT("");
 }
 
+inline const TCHAR* GetTextureCreateFlagString(ETextureCreateFlags TextureCreateFlag)
+{
+	switch (TextureCreateFlag)
+	{
+	case ETextureCreateFlags::None:
+		return TEXT("None");
+	case ETextureCreateFlags::RenderTargetable:
+		return TEXT("RenderTargetable");
+	case ETextureCreateFlags::ResolveTargetable:
+		return TEXT("ResolveTargetable");
+	case ETextureCreateFlags::DepthStencilTargetable:
+		return TEXT("DepthStencilTargetable");
+	case ETextureCreateFlags::ShaderResource:
+		return TEXT("ShaderResource");
+	case ETextureCreateFlags::SRGB:
+		return TEXT("SRGB");
+	case ETextureCreateFlags::CPUWritable:
+		return TEXT("CPUWritable");
+	case ETextureCreateFlags::NoTiling:
+		return TEXT("NoTiling");
+	case ETextureCreateFlags::VideoDecode:
+		return TEXT("VideoDecode");
+	case ETextureCreateFlags::Dynamic:
+		return TEXT("Dynamic");
+	case ETextureCreateFlags::InputAttachmentRead:
+		return TEXT("InputAttachmentRead");
+	case ETextureCreateFlags::Foveation:
+		return TEXT("Foveation");
+	case ETextureCreateFlags::Tiling3D:
+		return TEXT("Tiling3D");
+	case ETextureCreateFlags::Memoryless:
+		return TEXT("Memoryless");
+	case ETextureCreateFlags::GenerateMipCapable:
+		return TEXT("GenerateMipCapable");
+	case ETextureCreateFlags::FastVRAMPartialAlloc:
+		return TEXT("FastVRAMPartialAlloc");
+	case ETextureCreateFlags::DisableSRVCreation:
+		return TEXT("DisableSRVCreation");
+	case ETextureCreateFlags::DisableDCC:
+		return TEXT("DisableDCC");
+	case ETextureCreateFlags::UAV:
+		return TEXT("UAV");
+	case ETextureCreateFlags::Presentable:
+		return TEXT("Presentable");
+	case ETextureCreateFlags::CPUReadback:
+		return TEXT("CPUReadback");
+	case ETextureCreateFlags::OfflineProcessed:
+		return TEXT("OfflineProcessed");
+	case ETextureCreateFlags::FastVRAM:
+		return TEXT("FastVRAM");
+	case ETextureCreateFlags::HideInVisualizeTexture:
+		return TEXT("HideInVisualizeTexture");
+	case ETextureCreateFlags::Virtual:
+		return TEXT("Virtual");
+	case ETextureCreateFlags::TargetArraySlicesIndependently:
+		return TEXT("TargetArraySlicesIndependently");
+	case ETextureCreateFlags::Shared:
+		return TEXT("Shared");
+	case ETextureCreateFlags::NoFastClear:
+		return TEXT("NoFastClear");
+	case ETextureCreateFlags::DepthStencilResolveTarget:
+		return TEXT("DepthStencilResolveTarget");
+	case ETextureCreateFlags::Streamable:
+		return TEXT("Streamable");
+	case ETextureCreateFlags::NoFastClearFinalize:
+		return TEXT("NoFastClearFinalize");
+	case ETextureCreateFlags::AFRManual:
+		return TEXT("AFRManual");
+	case ETextureCreateFlags::ReduceMemoryWithTilingMode:
+		return TEXT("ReduceMemoryWithTilingMode");
+	case ETextureCreateFlags::Transient:
+		return TEXT("Transient");
+	}
+	return TEXT("");
+}
+
+inline const TCHAR* GetBufferUsageFlagString(EBufferUsageFlags BufferUsage)
+{
+	switch (BufferUsage)
+	{
+	case EBufferUsageFlags::None:
+		return TEXT("None");
+	case EBufferUsageFlags::Static:
+		return TEXT("Static");
+	case EBufferUsageFlags::Dynamic:
+		return TEXT("Dynamic");
+	case EBufferUsageFlags::Volatile:
+		return TEXT("Volatile");
+	case EBufferUsageFlags::UnorderedAccess:
+		return TEXT("UnorderedAccess");
+	case EBufferUsageFlags::ByteAddressBuffer:
+		return TEXT("ByteAddressBuffer");
+	case EBufferUsageFlags::SourceCopy:
+		return TEXT("SourceCopy");
+	case EBufferUsageFlags::StreamOutput:
+		return TEXT("StreamOutput");
+	case EBufferUsageFlags::DrawIndirect:
+		return TEXT("DrawIndirect");
+	case EBufferUsageFlags::ShaderResource:
+		return TEXT("ShaderResource");
+	case EBufferUsageFlags::KeepCPUAccessible:
+		return TEXT("KeepCPUAccessible");
+	case EBufferUsageFlags::FastVRAM:
+		return TEXT("FastVRAM");
+	case EBufferUsageFlags::Transient:
+		return TEXT("Transient");
+	case EBufferUsageFlags::Shared:
+		return TEXT("Shared");
+	case EBufferUsageFlags::AccelerationStructure:
+		return TEXT("AccelerationStructure");
+	case EBufferUsageFlags::VertexBuffer:
+		return TEXT("VertexBuffer");
+	case EBufferUsageFlags::IndexBuffer:
+		return TEXT("IndexBuffer");
+	case EBufferUsageFlags::StructuredBuffer:
+		return TEXT("StructuredBuffer");
+	}
+	return TEXT("");
+}
+
+inline const TCHAR* GetUniformBufferBaseTypeString(EUniformBufferBaseType BaseType)
+{
+	switch (BaseType)
+	{
+	case UBMT_INVALID:
+		return TEXT("UBMT_INVALID");
+	case UBMT_BOOL:
+		return TEXT("UBMT_BOOL");
+	case UBMT_INT32:
+		return TEXT("UBMT_INT32");
+	case UBMT_UINT32:
+		return TEXT("UBMT_UINT32");
+	case UBMT_FLOAT32:
+		return TEXT("UBMT_FLOAT32");
+	case UBMT_TEXTURE:
+		return TEXT("UBMT_TEXTURE");
+	case UBMT_SRV:
+		return TEXT("UBMT_SRV");
+	case UBMT_UAV:
+		return TEXT("UBMT_UAV");
+	case UBMT_SAMPLER:
+		return TEXT("UBMT_SAMPLER");
+	case UBMT_RDG_TEXTURE:
+		return TEXT("UBMT_RDG_TEXTURE");
+	case UBMT_RDG_TEXTURE_ACCESS:
+		return TEXT("UBMT_RDG_TEXTURE_ACCESS");
+	case UBMT_RDG_TEXTURE_ACCESS_ARRAY:
+		return TEXT("UBMT_RDG_TEXTURE_ACCESS_ARRAY");
+	case UBMT_RDG_TEXTURE_SRV:
+		return TEXT("UBMT_RDG_TEXTURE_SRV");
+	case UBMT_RDG_TEXTURE_UAV:
+		return TEXT("UBMT_RDG_TEXTURE_UAV");
+	case UBMT_RDG_BUFFER_ACCESS:
+		return TEXT("UBMT_RDG_BUFFER_ACCESS");
+	case UBMT_RDG_BUFFER_ACCESS_ARRAY:
+		return TEXT("UBMT_RDG_BUFFER_ACCESS_ARRAY");
+	case UBMT_RDG_BUFFER_SRV:
+		return TEXT("UBMT_RDG_BUFFER_SRV");
+	case UBMT_RDG_BUFFER_UAV:
+		return TEXT("UBMT_RDG_BUFFER_UAV");
+	case UBMT_RDG_UNIFORM_BUFFER:
+		return TEXT("UBMT_RDG_UNIFORM_BUFFER");
+	case UBMT_NESTED_STRUCT:
+		return TEXT("UBMT_NESTED_STRUCT");
+	case UBMT_INCLUDED_STRUCT:
+		return TEXT("UBMT_INCLUDED_STRUCT");
+	case UBMT_REFERENCED_STRUCT:
+		return TEXT("UBMT_REFERENCED_STRUCT");
+	case UBMT_RENDER_TARGET_BINDING_SLOTS:
+		return TEXT("UBMT_RENDER_TARGET_BINDING_SLOTS");
+	}
+	return TEXT("");
+}
+
+
 inline bool IsGeometryPipelineShaderFrequency(EShaderFrequency Frequency)
 {
 	return Frequency == SF_Mesh || Frequency == SF_Amplification;
diff --git a/Engine/Source/Runtime/RHI/Public/RHIUtilities.h b/Engine/Source/Runtime/RHI/Public/RHIUtilities.h
--- a/Engine/Source/Runtime/RHI/Public/RHIUtilities.h
+++ b/Engine/Source/Runtime/RHI/Public/RHIUtilities.h
@@ -136,6 +136,41 @@ static bool IsSnormFormat(EPixelFormat Format)
 	return false;
 }
 
+static bool IsUintFormat(EPixelFormat Format)
+{
+	switch (Format)
+	{
+	case PF_R32_UINT:
+	case PF_R16_UINT:
+	case PF_R16G16B16A16_UINT:
+	case PF_R32G32B32A32_UINT:
+	case PF_R16G16_UINT:
+	case PF_R8_UINT:
+	case PF_R8G8B8A8_UINT:
+	case PF_R32G32_UINT:
+		return true;
+
+	default:
+		break;
+	}
+	return false;
+}
+
+static bool IsSintFormat(EPixelFormat Format)
+{
+	switch (Format)
+	{
+	case PF_R32_SINT:
+	case PF_R16_SINT:
+	case PF_R16G16B16A16_SINT:
+		return true;
+
+	default:
+		break;
+	}
+	return false;
+}
+
 /** Get the best default resource state for the given texture creation flags */
 extern RHI_API ERHIAccess RHIGetDefaultResourceState(ETextureCreateFlags InUsage, bool bInHasInitialData);
 
diff --git a/Engine/Source/Runtime/RenderCore/Private/RenderGraphBuilder.cpp b/Engine/Source/Runtime/RenderCore/Private/RenderGraphBuilder.cpp
--- a/Engine/Source/Runtime/RenderCore/Private/RenderGraphBuilder.cpp
+++ b/Engine/Source/Runtime/RenderCore/Private/RenderGraphBuilder.cpp
@@ -8,6 +8,7 @@
 #include "VisualizeTexture.h"
 #include "ProfilingDebugging/CsvProfiler.h"
 
+
 #if ENABLE_RHI_VALIDATION
 
 inline void GatherPassUAVsForOverlapValidation(const FRDGPass* Pass, TArray<FRHIUnorderedAccessView*, TInlineAllocator<MaxSimultaneousUAVs, FRDGArrayAllocator>>& OutUAVs)
@@ -2562,8 +2563,16 @@ void FRDGBuilder::ExecutePass(FRDGPass* Pass, FRHIComputeCommandList& RHICmdList
 	Pass->GPUScopeOpsPrologue.Execute(RHICmdListPass);
 #endif
 
+#if RDG_DUMP_RESOURCES_AT_EACH_DRAW
+	BeginPassDump(Pass);
+#endif
+
 	Pass->Execute(RHICmdListPass);
 
+#if RDG_DUMP_RESOURCES_AT_EACH_DRAW
+	EndPassDump(Pass);
+#endif
+
 #if RDG_GPU_SCOPES
 	Pass->GPUScopeOpsEpilogue.Execute(RHICmdListPass);
 #endif
diff --git a/Engine/Source/Runtime/RenderCore/Private/RenderGraphDumpResources.cpp b/Engine/Source/Runtime/RenderCore/Private/RenderGraphDumpResources.cpp
--- a/Engine/Source/Runtime/RenderCore/Private/RenderGraphDumpResources.cpp
+++ b/Engine/Source/Runtime/RenderCore/Private/RenderGraphDumpResources.cpp
@@ -2,12 +2,19 @@
 
 #include "RenderGraph.h"
 #include "HAL/PlatformFileManager.h"
+#include "HAL/PlatformMisc.h"
 #include "Misc/FileHelper.h"
 #include "Misc/WildcardString.h"
+#include "Misc/OutputDeviceRedirector.h"
+#include "Runtime/Launch/Resources/Version.h"
 #include "BuildSettings.h"
 #include "Serialization/JsonTypes.h"
 #include "Serialization/JsonReader.h"
 #include "Serialization/JsonSerializer.h"
+#include "GenericPlatform/GenericPlatformCrashContext.h"
+#include "GenericPlatform/GenericPlatformDriver.h"
+#include "RenderGraphUtils.h"
+#include "GlobalShader.h"
 
 #if RDG_DUMP_RESOURCES
 
@@ -20,8 +27,94 @@ static TAutoConsoleVariable<FString> GDumpGPURootCVar(
 	TEXT("Allows to filter the tree when using r.DumpGPU command, the pattern match is case sensitive."),
 	ECVF_Default);
 
+static TAutoConsoleVariable<int32> GDumpTextureCVar(
+	TEXT("r.DumpGPU.Texture"), 2,
+	TEXT("Whether to dump textures.\n")
+	TEXT(" 0: Ignores all textures\n")
+	TEXT(" 1: Dump only textures' descriptors\n")
+	TEXT(" 2: Dump textures' descriptors and binaries (default)"),
+	ECVF_RenderThreadSafe);
+
+static TAutoConsoleVariable<int32> GDumpBufferCVar(
+	TEXT("r.DumpGPU.Buffer"), 2,
+	TEXT("Whether to dump buffer.\n")
+	TEXT(" 0: Ignores all buffers\n")
+	TEXT(" 1: Dump only buffers' descriptors\n")
+	TEXT(" 2: Dump buffers' descriptors and binaries (default)"),
+	ECVF_RenderThreadSafe);
+
+static TAutoConsoleVariable<int32> GDumpGPUPassParameters(
+	TEXT("r.DumpGPU.PassParameters"), 1,
+	TEXT("Whether to dump the pass parameters."),
+	ECVF_RenderThreadSafe);
+
+static TAutoConsoleVariable<int32> GDumpGPUDraws(
+	TEXT("r.DumpGPU.Draws"), 0,
+	TEXT("Whether to dump resource after each individual draw call (disabled by default)."),
+	ECVF_Default);
+
+static TAutoConsoleVariable<int32> GDumpExploreCVar(
+	TEXT("r.DumpGPU.Explore"), 1,
+	TEXT("Whether to open file explorer to where the GPU dump on completion (enabled by default)."),
+	ECVF_Default);
+
+static TAutoConsoleVariable<int32> GDumpRenderingConsoleVariablesCVar(
+	TEXT("r.DumpGPU.ConsoleVariables"), 1,
+	TEXT("Whether to dump rendering console variables (enabled by default)."),
+	ECVF_Default);
+
+static TAutoConsoleVariable<int32> GDumpTestEnableDiskWrite(
+	TEXT("r.DumpGPU.Test.EnableDiskWrite"), 1,
+	TEXT("Master switch whether any files should be written to disk, used for r.DumpGPU automation tests to not fill up workers' hard drive."),
+	ECVF_Default);
+
+static TAutoConsoleVariable<FString> GDumpGPUDirectoryCVar(
+	TEXT("r.DumpGPU.Directory"), TEXT(""),
+	TEXT("Directory to dump to."),
+	ECVF_Default);
+
+// Although this cvar does not seams used in the C++ code base, it is dumped by DumpRenderingCVarsToCSV() and used by GPUDumpViewer.html.
+static TAutoConsoleVariable<FString> GDumpGPUVisualizeResource(
+	TEXT("r.DumpGPU.Viewer.Visualize"), TEXT(""),
+	TEXT("Name of RDG output resource to automatically open in the dump viewer."),
+	ECVF_Default);
+
+class FDumpTextureCS : public FGlobalShader
+{
+	DECLARE_GLOBAL_SHADER(FDumpTextureCS);
+	SHADER_USE_PARAMETER_STRUCT(FDumpTextureCS, FGlobalShader);
+
+	static inline bool IsSupported(const FStaticShaderPlatform ShaderPlatform)
+	{
+		return RHISupportsComputeShaders(ShaderPlatform);
+	}
+
+	enum class ETextureType
+	{
+		Texture2DFloatNoMSAA,
+		Texture2DUintNoMSAA,
+		Texture2DDepthStencilNoMSAA,
+		MAX
+	};
+	class FTextureTypeDim : SHADER_PERMUTATION_ENUM_CLASS("DIM_TEXTURE_TYPE", ETextureType);
+	using FPermutationDomain = TShaderPermutationDomain<FTextureTypeDim>;
+
+	static bool ShouldCompilePermutation(const FGlobalShaderPermutationParameters& Parameters)
+	{
+		return IsSupported(Parameters.Platform);
+	}
+
+	BEGIN_SHADER_PARAMETER_STRUCT(FParameters, )
+		SHADER_PARAMETER_SRV(Texture2D, Texture)
+		SHADER_PARAMETER_UAV(RWTexture2D, StagingOutput)
+	END_SHADER_PARAMETER_STRUCT()
+};
+
+IMPLEMENT_GLOBAL_SHADER(FDumpTextureCS, "/Engine/Private/Tools/DumpTexture.usf", "MainCS", SF_Compute);
+
 BEGIN_SHADER_PARAMETER_STRUCT(FDumpTexturePass, )
-	RDG_TEXTURE_ACCESS(Texture, ERHIAccess::CopySrc)
+	SHADER_PARAMETER_RDG_TEXTURE_SRV(Texture2D, Texture)
+	RDG_TEXTURE_ACCESS_DYNAMIC(TextureAccess)
 END_SHADER_PARAMETER_STRUCT()
 
 BEGIN_SHADER_PARAMETER_STRUCT(FDumpBufferPass, )
@@ -30,33 +123,300 @@ END_SHADER_PARAMETER_STRUCT()
 
 struct FRDGResourceDumpContext
 {
+	static constexpr const TCHAR* kBaseDir = TEXT("Base/");
+	static constexpr const TCHAR* kPassesDir = TEXT("Passes/");
+	static constexpr const TCHAR* kResourcesDir = TEXT("Resources/");
+	static constexpr const TCHAR* kStructuresDir = TEXT("Structures/");
+	static constexpr const TCHAR* kStructuresMetadataDir = TEXT("StructuresMetadata/");
+
+	bool bEnableDiskWrite = false;
 	FString DumpingDirectoryPath;
 	FDateTime Time;
-	int32 ResourcesDumped = 0;
+	FGenericPlatformMemoryConstants MemoryConstants;
+	FGenericPlatformMemoryStats MemoryStats;
+	int32 ResourcesDumpPasses = 0;
+	int32 ResourcesDumpExecutedPasses = 0;
 	int32 PassesCount = 0;
 	TMap<const FRDGResource*, const FRDGPass*> LastResourceVersion;
+	TSet<const void*> IsDumpedToDisk;
+
+	// Pass being dumping individual draws
+	const FRDGPass* DrawDumpingPass = nullptr;
+	int32 DrawDumpCount = 0;
 
 	bool bShowInExplore = false;
 
 	bool IsDumpingFrame() const
 	{
-		check(IsInRenderingThread());
+		check(IsInRenderingThread() || IsInParallelRenderingThread());
 		return !DumpingDirectoryPath.IsEmpty();
 	}
 
-	static FString GetResourceName(const FRDGResource* Resource)
+	FString GetDumpFullPath(const FString& DumpRelativeFileName) const
+	{
+		check(bEnableDiskWrite);
+		return DumpingDirectoryPath / DumpRelativeFileName;
+	}
+
+	bool DumpStringToFile(FStringView OutputString, const FString& FileName, uint32 WriteFlags = FILEWRITE_None)
+	{
+		// Make it has if the write happened and was successful.
+		if (!bEnableDiskWrite)
+		{
+			return true;
+		}
+
+		FString FullPath = GetDumpFullPath(FileName);
+		return FFileHelper::SaveStringToFile(
+			OutputString, *FullPath,
+			FFileHelper::EEncodingOptions::AutoDetect, &IFileManager::Get(), WriteFlags);
+	}
+
+	bool DumpJsonToFile(const TSharedPtr<FJsonObject>& JsonObject, const FString& FileName, uint32 WriteFlags = FILEWRITE_None)
+	{
+		FString OutputString;
+		TSharedRef<TJsonWriter<TCHAR, TPrettyJsonPrintPolicy<TCHAR>>> Writer = TJsonWriterFactory<TCHAR, TPrettyJsonPrintPolicy<TCHAR>>::Create(&OutputString);
+		FJsonSerializer::Serialize(JsonObject.ToSharedRef(), Writer);
+
+		return DumpStringToFile(OutputString, FileName, WriteFlags);
+	}
+
+	bool DumpBinaryToFile(TArrayView<const uint8> ArrayView, const FString& FileName)
+	{
+		// Make it has if the write happened and was successful.
+		if (!bEnableDiskWrite)
+		{
+			return true;
+		}
+
+		FString FullPath = GetDumpFullPath(FileName);
+		return FFileHelper::SaveArrayToFile(ArrayView, *FullPath);
+	}
+
+	bool IsUnsafeToDumpResource(uint32 ResourceByteSize, float DumpMemoryMultiplier) const
+	{
+		uint64 AproximatedStagingMemoryRequired = uint64(double(ResourceByteSize) * DumpMemoryMultiplier);
+		uint64 MaxMemoryAvailable = FMath::Min(MemoryStats.AvailablePhysical, MemoryStats.AvailableVirtual);
+
+		return AproximatedStagingMemoryRequired > MaxMemoryAvailable;
+	}
+
+	template<typename T>
+	static uint64 PtrToUint(const T* Ptr)
+	{
+		return static_cast<uint64>(reinterpret_cast<size_t>(Ptr));
+	}
+
+	template<typename T>
+	static FString PtrToString(const T* Ptr)
+	{
+		return FString::Printf(TEXT("%016x"), PtrToUint(Ptr));
+	}
+
+	static FString GetUniqueResourceName(const FRDGResource* Resource)
+	{
+		FString UniqueResourceName = FString::Printf(TEXT("%s.%016x"), Resource->Name, PtrToUint(Resource));
+		UniqueResourceName.ReplaceInline(TEXT("/"), TEXT(""));
+		UniqueResourceName.ReplaceInline(TEXT("\\"), TEXT(""));
+		return UniqueResourceName;
+	}
+
+	static FString GetUniqueSubResourceName(FRDGTextureSRVDesc SubResourceDesc)
+	{
+		check(SubResourceDesc.NumMipLevels == 1);
+
+		FString UniqueResourceName = GetUniqueResourceName(SubResourceDesc.Texture);
+
+		if (SubResourceDesc.Format == PF_X24_G8)
+		{
+			return FString::Printf(TEXT("%s.stencil"), *UniqueResourceName);
+		}
+
+		return FString::Printf(TEXT("%s.mip%d"), *UniqueResourceName, SubResourceDesc.MipLevel);
+	}
+
+	void ReleaseRHIResources(FRHICommandListImmediate& RHICmdList)
+	{
+		// Flush the RHI resource memory so the readback memory can be fully reused in the next resource dump.
+		{
+			RHICmdList.SubmitCommandsAndFlushGPU();
+			RHICmdList.BlockUntilGPUIdle();
+			FRHIResource::FlushPendingDeletes(RHICmdList);
+			RHICmdList.FlushResources();
+			RHICmdList.ImmediateFlush(EImmediateFlushType::FlushRHIThreadFlushResources);
+		}
+	}
+
+	void UpdatePassProgress()
+	{
+		ResourcesDumpExecutedPasses++;
+
+		if (ResourcesDumpExecutedPasses % 10 == 0)
+		{
+			UE_LOG(LogRendererCore, Display, TEXT("Dumped %d / %d resources"), ResourcesDumpExecutedPasses, ResourcesDumpPasses);
+		}
+	}
+
+	void DumpRenderingCVarsToCSV() const
 	{
-		return FString::Printf(TEXT("%s.%p"), Resource->Name, Resource);
+		FString FileName = GetDumpFullPath(FString(FRDGResourceDumpContext::kBaseDir) / TEXT("ConsoleVariables.csv"));
+
+		TUniquePtr<FArchive> Ar = TUniquePtr<FArchive>(IFileManager::Get().CreateFileWriter(*FileName));
+		auto OnConsoleVariable = [&Ar](const TCHAR* CVarName, IConsoleObject* ConsoleObj)
+		{
+			if (ConsoleObj->TestFlags(ECVF_Unregistered))
+			{
+				return;
+			}
+
+			const IConsoleVariable* CVar = ConsoleObj->AsVariable();
+			if (!CVar)
+			{
+				return;
+			}
+
+			EConsoleVariableFlags CVarFlags = CVar->GetFlags();
+
+			const TCHAR* Type = nullptr;
+			if (CVar->IsVariableBool())
+			{
+				Type = TEXT("bool");
+			}
+			else if (CVar->IsVariableInt())
+			{
+				Type = TEXT("int32");
+			}
+			else if (CVar->IsVariableFloat())
+			{
+				Type = TEXT("float");
+			}
+			else if (CVar->IsVariableString())
+			{
+				Type = TEXT("FString");
+			}
+			else
+			{
+				return;
+			}
+
+			const TCHAR* SetByName = TEXT("");
+			switch (CVarFlags & ECVF_SetByMask)
+			{
+			case ECVF_SetByConstructor:
+			{
+				SetByName = TEXT("Constructor");
+				break;
+			}
+			case ECVF_SetByScalability:
+			{
+				SetByName = TEXT("Scalability");
+				break;
+			}
+			case ECVF_SetByGameSetting:
+			{
+				SetByName = TEXT("GameSetting");
+				break;
+			}
+			case ECVF_SetByProjectSetting:
+			{
+				SetByName = TEXT("ProjectSetting");
+				break;
+			}
+			case ECVF_SetBySystemSettingsIni:
+			{
+				SetByName = TEXT("SystemSettingsIni");
+				break;
+			}
+			case ECVF_SetByDeviceProfile:
+			{
+				SetByName = TEXT("DeviceProfile");
+				break;
+			}
+			case ECVF_SetByConsoleVariablesIni:
+			{
+				SetByName = TEXT("ConsoleVariablesIni");
+				break;
+			}
+			case ECVF_SetByCommandline:
+			{
+				SetByName = TEXT("Commandline");
+				break;
+			}
+			case ECVF_SetByCode:
+			{
+				SetByName = TEXT("Code");
+				break;
+			}
+			case ECVF_SetByConsole:
+			{
+				SetByName = TEXT("Console");
+				break;
+			}
+			default:
+				unimplemented();
+			}
+
+			FString Value = CVar->GetString();
+
+			FString CSVLine = FString::Printf(TEXT("%s,%s,%s,%s\n"), CVarName, Type, SetByName, *Value);
+			FStringView CSVLineView(CSVLine);
+
+			auto Src = StringCast<ANSICHAR>(CSVLineView.GetData(), CSVLineView.Len());
+			Ar->Serialize((ANSICHAR*)Src.Get(), Src.Length() * sizeof(ANSICHAR));
+		};
+
+		bool bSuccess = false;
+		if (Ar)
+		{
+			{
+				FString CSVLine = TEXT("CVar,Type,SetBy,Value\n");
+				FStringView CSVLineView(CSVLine);
+
+				auto Src = StringCast<ANSICHAR>(CSVLineView.GetData(), CSVLineView.Len());
+				Ar->Serialize((ANSICHAR*)Src.Get(), Src.Length() * sizeof(ANSICHAR));
+			}
+
+			if (GDumpRenderingConsoleVariablesCVar.GetValueOnGameThread() != 0)
+			{
+				IConsoleManager::Get().ForEachConsoleObjectThatStartsWith(FConsoleObjectVisitor::CreateLambda(OnConsoleVariable), TEXT("r."));
+				IConsoleManager::Get().ForEachConsoleObjectThatStartsWith(FConsoleObjectVisitor::CreateLambda(OnConsoleVariable), TEXT("sg."));
+				IConsoleManager::Get().ForEachConsoleObjectThatContains(FConsoleObjectVisitor::CreateLambda(OnConsoleVariable), TEXT("rhi"));
+			}
+			else
+			{
+				IConsoleManager::Get().ForEachConsoleObjectThatStartsWith(FConsoleObjectVisitor::CreateLambda(OnConsoleVariable), TEXT("r.DumpGPU."));
+			}
+			// Always explicitly close to catch errors from flush/close
+			Ar->Close();
+
+			bSuccess = !Ar->IsError() && !Ar->IsCriticalError();
+		}
+
+		if (bSuccess)
+		{
+			UE_LOG(LogRendererCore, Display, TEXT("DumpGPU dumped rendering cvars to %s."), *FileName);
+		}
+		else
+		{
+			UE_LOG(LogRendererCore, Error, TEXT("DumpGPU had a file error when dumping rendering cvars to %s."), *FileName);
+		}
 	}
 
-	static FString GetResourceVersionDumpName(const FRDGPass* Pass, const FRDGResource* Resource)
+	template<typename T>
+	bool IsDumped(const T* Ptr) const
 	{
-		return FString::Printf(TEXT("%s.%p.v%p.bin"), Resource->Name, Resource, Pass);
+		return IsDumpedToDisk.Contains(static_cast<const void*>(Ptr));
 	}
 
-	FString GetResourceDumpDirectory() const
+	template<typename T>
+	void SetDumped(const T* Ptr)
 	{
-		return DumpingDirectoryPath + TEXT("Resources/");
+		check(!IsDumped(Ptr));
+		if (IsDumpedToDisk.Num() % 1024 == 0)
+		{
+			IsDumpedToDisk.Reserve(IsDumpedToDisk.Num() + 1024);
+		}
+		IsDumpedToDisk.Add(static_cast<const void*>(Ptr));
 	}
 
 	void GetResourceDumpInfo(
@@ -64,14 +424,15 @@ struct FRDGResourceDumpContext
 		const FRDGResource* Resource,
 		bool bIsOutputResource,
 		bool* bOutDumpResourceInfos,
-		FString* OutResourceVersionDumpName)
+		bool* bOutDumpResourceBinary)
 	{
 		*bOutDumpResourceInfos = false;
+		*bOutDumpResourceBinary = bIsOutputResource;
 
-		if (bIsOutputResource)
-		{
-			*OutResourceVersionDumpName = GetResourceVersionDumpName(Pass, Resource);
-		}
+		//if (bIsOutputResource)
+		//{
+		//	*OutResourceVersionDumpName = GetResourceVersionDumpName(Pass, Resource);
+		//}
 
 		if (!LastResourceVersion.Contains(Resource))
 		{
@@ -81,7 +442,7 @@ struct FRDGResourceDumpContext
 			// If not an output, it might be a resource undumped by r.DumpGPU.Root or external texture so still dump it as v0.
 			if (!bIsOutputResource)
 			{
-				*OutResourceVersionDumpName = GetResourceVersionDumpName(/* Pass = */ nullptr, Resource);
+				*bOutDumpResourceBinary = true;
 			}
 
 			if (LastResourceVersion.Num() % 1024 == 0)
@@ -96,133 +457,535 @@ struct FRDGResourceDumpContext
 		}
 	}
 
-	void AddDumpTexturePass(
-		FRDGBuilder& GraphBuilder,
-		TArray<TSharedPtr<FJsonValue>>& InputResourceNames,
-		TArray<TSharedPtr<FJsonValue>>& OutputResourceNames,
-		const FRDGPass* Pass,
-		FRDGTexture* Texture,
-		bool bIsOutputResource)
+	FString ToJson(EPixelFormat Format)
+	{
+		FString PixelFormat = GPixelFormats[Format].Name;
+		if (!PixelFormat.StartsWith(TEXT("PF_")))
+		{
+			PixelFormat = FString::Printf(TEXT("PF_%s"), GPixelFormats[Format].Name);
+		}
+		return PixelFormat;
+	}
+
+	TSharedPtr<FJsonObject> ToJson(const FShaderParametersMetadata::FMember& Member)
 	{
-		FString UniqueResourceName = GetResourceName(Texture);
+		TSharedPtr<FJsonObject> JsonObject = MakeShareable(new FJsonObject);
+		JsonObject->SetStringField(TEXT("Name"), Member.GetName());
+		JsonObject->SetStringField(TEXT("ShaderType"), Member.GetShaderType());
+		JsonObject->SetNumberField(TEXT("FileLine"), Member.GetFileLine());
+		JsonObject->SetNumberField(TEXT("Offset"), Member.GetOffset());
+		JsonObject->SetStringField(TEXT("BaseType"), GetUniformBufferBaseTypeString(Member.GetBaseType()));
+		JsonObject->SetNumberField(TEXT("Precision"), Member.GetPrecision());
+		JsonObject->SetNumberField(TEXT("NumRows"), Member.GetNumRows());
+		JsonObject->SetNumberField(TEXT("NumColumns"), Member.GetNumColumns());
+		JsonObject->SetNumberField(TEXT("NumElements"), Member.GetNumElements());
+		JsonObject->SetStringField(TEXT("StructMetadata"), *PtrToString(Member.GetStructMetadata()));
+		return JsonObject;
+	}
+
+	TSharedPtr<FJsonObject> ToJson(const FShaderParametersMetadata* Metadata)
+	{
+		TSharedPtr<FJsonObject> JsonObject = MakeShareable(new FJsonObject);
+		JsonObject->SetStringField(TEXT("StructTypeName"), Metadata->GetStructTypeName());
+		JsonObject->SetStringField(TEXT("ShaderVariableName"), Metadata->GetShaderVariableName());
+		JsonObject->SetStringField(TEXT("FileName"), Metadata->GetFileName());
+		JsonObject->SetNumberField(TEXT("FileLine"), Metadata->GetFileLine());
+		JsonObject->SetNumberField(TEXT("Size"), Metadata->GetSize());
+		//JsonObject->SetNumberField(TEXT("UseCase"), Metadata->GetUseCase());
 
-		if (bIsOutputResource)
 		{
-			OutputResourceNames.AddUnique(MakeShareable(new FJsonValueString(UniqueResourceName)));
+			TArray<TSharedPtr<FJsonValue>> Members;
+			for (const FShaderParametersMetadata::FMember& Member : Metadata->GetMembers())
+			{
+				Members.Add(MakeShared<FJsonValueObject>(ToJson(Member)));
+			}
+			JsonObject->SetArrayField(TEXT("Members"), Members);
 		}
-		else
+
+		return JsonObject;
+	}
+
+	TSharedPtr<FJsonObject> ToJson(const TCHAR* Name, const FRDGTextureDesc& Desc)
+	{
+		FString PixelFormat = ToJson(Desc.Format);
+
+		int32 ResourceByteSize = Desc.Extent.X * Desc.Extent.Y * Desc.Depth * Desc.ArraySize * Desc.NumSamples * GPixelFormats[Desc.Format].BlockBytes;
+
+		TSharedPtr<FJsonObject> JsonObject = MakeShareable(new FJsonObject);
+		JsonObject->SetStringField(TEXT("Name"), Name);
+		JsonObject->SetNumberField(TEXT("ByteSize"), ResourceByteSize);
+		JsonObject->SetStringField(TEXT("Desc"), TEXT("FRDGTextureDesc"));
+		JsonObject->SetStringField(TEXT("Type"), GetTextureDimensionString(Desc.Dimension));
+		JsonObject->SetStringField(TEXT("Format"), *PixelFormat);
+		JsonObject->SetNumberField(TEXT("ExtentX"), Desc.Extent.X);
+		JsonObject->SetNumberField(TEXT("ExtentY"), Desc.Extent.Y);
+		JsonObject->SetNumberField(TEXT("Depth"), Desc.Depth);
+		JsonObject->SetNumberField(TEXT("ArraySize"), Desc.ArraySize);
+		JsonObject->SetNumberField(TEXT("NumMips"), Desc.NumMips);
+		JsonObject->SetNumberField(TEXT("NumSamples"), Desc.NumSamples);
+
 		{
-			InputResourceNames.AddUnique(MakeShareable(new FJsonValueString(UniqueResourceName)));
+			TArray<TSharedPtr<FJsonValue>> FlagsNames;
+			for (uint64 BitId = 0; BitId < uint64(8 * sizeof(ETextureCreateFlags)); BitId++)
+			{
+				ETextureCreateFlags Flag = ETextureCreateFlags(uint64(1) << BitId);
+				if (EnumHasAnyFlags(Desc.Flags, Flag))
+				{
+					FlagsNames.Add(MakeShareable(new FJsonValueString(GetTextureCreateFlagString(Flag))));
+				}
+			}
+			JsonObject->SetArrayField(TEXT("Flags"), FlagsNames);
 		}
 
-		const FRDGTextureDesc& Desc = Texture->Desc;
+		return JsonObject;
+	}
+
+	TSharedPtr<FJsonObject> ToJson(const TCHAR* Name, const FRDGBufferDesc& Desc, int32 ResourceByteSize)
+	{
+		TSharedPtr<FJsonObject> JsonObject = MakeShareable(new FJsonObject);
+		JsonObject->SetStringField(TEXT("Name"), Name);
+		JsonObject->SetNumberField(TEXT("ByteSize"), ResourceByteSize);
+		JsonObject->SetStringField(TEXT("Desc"), TEXT("FRDGBufferDesc"));
+		JsonObject->SetStringField(TEXT("Type"), GetBufferUnderlyingTypeName(Desc.UnderlyingType));
+		JsonObject->SetNumberField(TEXT("BytesPerElement"), Desc.BytesPerElement);
+		JsonObject->SetNumberField(TEXT("NumElements"), Desc.NumElements);
+		JsonObject->SetStringField(TEXT("Metadata"), *PtrToString(Desc.Metadata));
+
+		{
+			TArray<TSharedPtr<FJsonValue>> UsageNames;
+			for (uint64 BitId = 0; BitId < uint64(8 * sizeof(EBufferUsageFlags)); BitId++)
+			{
+				EBufferUsageFlags Flag = EBufferUsageFlags(uint64(1) << BitId);
+				if (EnumHasAnyFlags(Desc.Usage, Flag))
+				{
+					UsageNames.Add(MakeShareable(new FJsonValueString(GetBufferUsageFlagString(Flag))));
+				}
+			}
+			JsonObject->SetArrayField(TEXT("Usage"), UsageNames);
+		}
+		return JsonObject;
+	}
+
+	void Dump(const FShaderParametersMetadata* Metadata)
+	{
+		if (IsDumped(Metadata))
+		{
+			return;
+		}
+
+		TSharedPtr<FJsonObject> JsonObject = ToJson(Metadata);
+		FString JsonPath = kStructuresMetadataDir / PtrToString(Metadata) + TEXT(".json");
+		DumpJsonToFile(JsonObject, JsonPath);
+		SetDumped(Metadata);
+
+		// Dump dependencies
+		Metadata->IterateStructureMetadataDependencies(
+		[&](const FShaderParametersMetadata* Struct)
+		{
+			if (Struct)
+			{
+				Dump(Struct);
+			}
+		});
+	}
+
+	struct FTextureSubresourceDumpDesc
+	{
+		FIntPoint SubResourceExtent = FIntPoint(0, 0);
 		int32 ByteSize = 0;
+		bool bPreprocessForStaging = false;
+		FDumpTextureCS::ETextureType DumpTextureType = FDumpTextureCS::ETextureType::MAX;
+		EPixelFormat PreprocessedPixelFormat;
+
+		bool IsDumpSupported() const
+		{
+			return ByteSize != 0;
+		}
+	};
+	
+	FTextureSubresourceDumpDesc TranslateSubresourceDumpDesc(FRDGTextureSRVDesc SubresourceDesc)
+	{
+		const FRDGTextureDesc& Desc = SubresourceDesc.Texture->Desc;
+
+		FTextureSubresourceDumpDesc SubresourceDumpDesc;
+		SubresourceDumpDesc.PreprocessedPixelFormat = Desc.Format;
+
+		bool bIsUnsupported = false;
+
+		if (GPixelFormats[Desc.Format].BlockSizeX != 1 ||
+			GPixelFormats[Desc.Format].BlockSizeY != 1 ||
+			GPixelFormats[Desc.Format].BlockSizeZ != 1)
+		{
+			bIsUnsupported = true;
+		}
+
+		if (!bIsUnsupported && Desc.IsTexture2D() && !Desc.IsMultisample() && !Desc.IsTextureArray())
 		{
-			bool bIsUnsupported = false;
+			SubresourceDumpDesc.SubResourceExtent.X = Desc.Extent.X >> SubresourceDesc.MipLevel;
+			SubresourceDumpDesc.SubResourceExtent.Y = Desc.Extent.Y >> SubresourceDesc.MipLevel;
 
-			if (GPixelFormats[Desc.Format].BlockSizeX != 1 || 
-				GPixelFormats[Desc.Format].BlockSizeY != 1 ||
-				GPixelFormats[Desc.Format].BlockSizeZ != 1)
+			if (IsUintFormat(Desc.Format) || IsSintFormat(Desc.Format))
+			{
+				SubresourceDumpDesc.DumpTextureType = FDumpTextureCS::ETextureType::Texture2DUintNoMSAA;
+			}
+			else
 			{
-				bIsUnsupported = true;
+				SubresourceDumpDesc.DumpTextureType = FDumpTextureCS::ETextureType::Texture2DFloatNoMSAA;
 			}
 
-			if (Desc.Format == PF_DepthStencil)
+			if (SubresourceDesc.Format == PF_X24_G8)
+			{
+				SubresourceDumpDesc.PreprocessedPixelFormat = PF_R8_UINT;
+				SubresourceDumpDesc.DumpTextureType = FDumpTextureCS::ETextureType::Texture2DDepthStencilNoMSAA;
+			}
+			else if (Desc.Format == PF_DepthStencil)
+			{
+				SubresourceDumpDesc.PreprocessedPixelFormat = PF_R32_FLOAT;
+				SubresourceDumpDesc.DumpTextureType = FDumpTextureCS::ETextureType::Texture2DFloatNoMSAA;
+			}
+			else if (Desc.Format == PF_ShadowDepth)
+			{
+				SubresourceDumpDesc.PreprocessedPixelFormat = PF_R32_FLOAT;
+				SubresourceDumpDesc.DumpTextureType = FDumpTextureCS::ETextureType::Texture2DFloatNoMSAA;
+			}
+			else if (Desc.Format == PF_D24)
 			{
-				bIsUnsupported = true;
+				SubresourceDumpDesc.PreprocessedPixelFormat = PF_R32_FLOAT;
+				SubresourceDumpDesc.DumpTextureType = FDumpTextureCS::ETextureType::Texture2DFloatNoMSAA;
 			}
+		}
+
+		SubresourceDumpDesc.ByteSize = SubresourceDumpDesc.SubResourceExtent.X * SubresourceDumpDesc.SubResourceExtent.Y * GPixelFormats[SubresourceDumpDesc.PreprocessedPixelFormat].BlockBytes;
+
+		// Whether the subresource need preprocessing pass before copy into staging.
+		SubresourceDumpDesc.bPreprocessForStaging = SubresourceDumpDesc.PreprocessedPixelFormat != Desc.Format || SubresourceDesc.MipLevel != 0;
 
-			if (!bIsUnsupported && Desc.IsTexture2D() && !Desc.IsMultisample() && !Desc.IsTextureArray())
+		return SubresourceDumpDesc;
+	}
+
+	void DumpTextureSubResource(
+		FRHICommandListImmediate& RHICmdList,
+		const TCHAR* TextureDebugName,
+		FRHITexture* Texture,
+		FRHIShaderResourceView* SubResourceSRV,
+		const FTextureSubresourceDumpDesc& SubresourceDumpDesc,
+		const FString& DumpFilePath)
+	{
+		check(IsInRenderingThread());
+
+		// Preprocess
+		FTextureRHIRef StagingSrcTexture;
+		if (SubresourceDumpDesc.bPreprocessForStaging)
+		{
 			{
-				ByteSize = Desc.Extent.X * Desc.Extent.Y * GPixelFormats[Desc.Format].BlockBytes;
+				FRHIResourceCreateInfo CreateInfo(TEXT("DumpGPU.PreprocessTexture"));
+				StagingSrcTexture = RHICreateTexture2D(
+					SubresourceDumpDesc.SubResourceExtent.X,
+					SubresourceDumpDesc.SubResourceExtent.Y,
+					SubresourceDumpDesc.PreprocessedPixelFormat,
+					/* NumMips = */ 1,
+					/* NumSamples = */ 1,
+					TexCreate_UAV | TexCreate_ShaderResource | TexCreate_HideInVisualizeTexture,
+					CreateInfo);
 			}
+
+			FUnorderedAccessViewRHIRef StagingOutput = RHICreateUnorderedAccessView(StagingSrcTexture, /* MipLevel = */ 0);
+
+			RHICmdList.Transition(FRHITransitionInfo(StagingSrcTexture, ERHIAccess::Unknown, ERHIAccess::UAVCompute));
+
+			FDumpTextureCS::FPermutationDomain PermutationVector;
+			PermutationVector.Set<FDumpTextureCS::FTextureTypeDim>(SubresourceDumpDesc.DumpTextureType);
+			TShaderMapRef<FDumpTextureCS> ComputeShader(GetGlobalShaderMap(GMaxRHIShaderPlatform), PermutationVector);
+
+			FDumpTextureCS::FParameters ShaderParameters;
+			ShaderParameters.Texture = SubResourceSRV;
+			ShaderParameters.StagingOutput = StagingOutput;
+			FComputeShaderUtils::Dispatch(
+				RHICmdList,
+				ComputeShader,
+				ShaderParameters,
+				FComputeShaderUtils::GetGroupCount(SubresourceDumpDesc.SubResourceExtent, 8));
+
+			RHICmdList.Transition(FRHITransitionInfo(StagingSrcTexture, ERHIAccess::UAVCompute, ERHIAccess::CopySrc));
+		}
+		else
+		{
+			StagingSrcTexture = Texture;
 		}
 
-		bool bDumpResourceInfos;
-		FString ResourceVersionDumpName;
-		GetResourceDumpInfo(Pass, Texture, bIsOutputResource, &bDumpResourceInfos, &ResourceVersionDumpName);
+		// Copy the texture for CPU readback
+		FTextureRHIRef StagingTexture;
+		{
+			FRHIResourceCreateInfo CreateInfo(TEXT("DumpGPU.StagingTexture"));
+			StagingTexture = RHICreateTexture2D(
+				SubresourceDumpDesc.SubResourceExtent.X,
+				SubresourceDumpDesc.SubResourceExtent.Y,
+				SubresourceDumpDesc.PreprocessedPixelFormat,
+				/* NumMips = */ 1,
+				/* NumSamples = */ 1,
+				TexCreate_CPUReadback | TexCreate_HideInVisualizeTexture,
+				CreateInfo);
+
+			RHICmdList.Transition(FRHITransitionInfo(StagingTexture, ERHIAccess::Unknown, ERHIAccess::CopyDest));
+
+			// Ensure this copy call does not perform any transitions. We're handling them manually.
+			FResolveParams ResolveParams;
+			ResolveParams.SourceAccessFinal = ERHIAccess::Unknown;
+			ResolveParams.DestAccessFinal = ERHIAccess::Unknown;
+
+			// Transfer memory GPU -> CPU
+			RHICmdList.CopyToResolveTarget(StagingSrcTexture, StagingTexture, ResolveParams);
+
+			RHICmdList.Transition(FRHITransitionInfo(StagingTexture, ERHIAccess::CopyDest, ERHIAccess::CPURead));
+		}
+
+		// Submit to GPU and wait for completion.
+		static const FName FenceName(TEXT("DumpGPU.TextureFence"));
+		FGPUFenceRHIRef Fence = RHICreateGPUFence(FenceName);
+		{
+			Fence->Clear();
+			RHICmdList.WriteGPUFence(Fence);
+			RHICmdList.SubmitCommandsAndFlushGPU();
+			RHICmdList.BlockUntilGPUIdle();
+		}
 
-		FString ResourceDumpDirectory = GetResourceDumpDirectory();
+		void* Content = nullptr;
+		int32 RowPitchInPixels = 0;
+		int32 ColumnPitchInPixels = 0;
+		RHICmdList.MapStagingSurface(StagingTexture, Fence.GetReference(), Content, RowPitchInPixels, ColumnPitchInPixels);
 
-		// Dump the information of the texture to json file.
-		if (bDumpResourceInfos)
+		if (Content)
 		{
-			FString PixelFormat = FString::Printf(TEXT("PF_%s"), GPixelFormats[Desc.Format].Name);
+			TArray<uint8> Array;
+			Array.SetNumUninitialized(SubresourceDumpDesc.ByteSize);
+
+			int32 BytePerPixel = GPixelFormats[SubresourceDumpDesc.PreprocessedPixelFormat].BlockBytes;
+
+			const uint8* SrcData = reinterpret_cast<const uint8*>(Content);
 
-			TSharedPtr<FJsonObject> JsonObject = MakeShareable(new FJsonObject);
-			JsonObject->SetStringField(TEXT("Name"), Texture->Name);
-			JsonObject->SetNumberField(TEXT("ByteSize"), ByteSize);
-			JsonObject->SetStringField(TEXT("Type"), GetTextureDimensionString(Desc.Dimension));
-			JsonObject->SetStringField(TEXT("Format"), *PixelFormat);
-			JsonObject->SetNumberField(TEXT("ExtentX"), Desc.Extent.X);
-			JsonObject->SetNumberField(TEXT("ExtentY"), Desc.Extent.Y);
-			JsonObject->SetNumberField(TEXT("Depth"), Desc.Depth);
-			JsonObject->SetNumberField(TEXT("ArraySize"), Desc.ArraySize);
-			JsonObject->SetNumberField(TEXT("NumMips"), Desc.NumMips);
-			JsonObject->SetNumberField(TEXT("NumSamples"), Desc.NumSamples);
+			for (int32 y = 0; y < SubresourceDumpDesc.SubResourceExtent.Y; y++)
+			{
+				// Flip the data to be bottom left corner for the WebGL viewer.
+				const uint8* SrcPos = SrcData + (SubresourceDumpDesc.SubResourceExtent.Y - 1 - y) * RowPitchInPixels * BytePerPixel;
+				uint8* DstPos = (&Array[0]) + y * SubresourceDumpDesc.SubResourceExtent.X * BytePerPixel;
+
+				FPlatformMemory::Memmove(DstPos, SrcPos, SubresourceDumpDesc.SubResourceExtent.X * BytePerPixel);
+			}
 
-			FString OutputString;
-			TSharedRef<TJsonWriter<TCHAR, TPrettyJsonPrintPolicy<TCHAR>>> Writer = TJsonWriterFactory<TCHAR, TPrettyJsonPrintPolicy<TCHAR>>::Create(&OutputString);
-			FJsonSerializer::Serialize(JsonObject.ToSharedRef(), Writer);
+			RHICmdList.UnmapStagingSurface(StagingTexture);
 
-			FString ResourceInfoJsonPath = ResourceDumpDirectory + UniqueResourceName + TEXT(".json");
-			FFileHelper::SaveStringToFile(OutputString, *ResourceInfoJsonPath);
+			TArrayView<const uint8> ArrayView(reinterpret_cast<const uint8*>(&Array[0]), SubresourceDumpDesc.ByteSize);
+			DumpBinaryToFile(ArrayView, DumpFilePath);
+		}
+		else
+		{
+			UE_LOG(LogRendererCore, Warning, TEXT("RHICmdList.MapStagingSurface() to dump texture %s failed."), TextureDebugName);
 		}
+	}
+
+	void DumpDrawTextureSubResource(
+		FRHICommandList& RHICmdList,
+		FRDGTextureSRVDesc SubresourceDesc,
+		ERHIAccess RHIAccessState)
+	{
+		check(IsInRenderingThread());
+
+		FRHICommandListImmediate& RHICmdListImmediate = FRHICommandListExecutor::GetImmediateCommandList();
+		check(&RHICmdListImmediate == &RHICmdList);
+
+		const FString UniqueResourceSubResourceName = GetUniqueSubResourceName(SubresourceDesc);
+		const FTextureSubresourceDumpDesc SubresourceDumpDesc = TranslateSubresourceDumpDesc(SubresourceDesc);
 
-		if (ByteSize == 0)
+		if (!SubresourceDumpDesc.IsDumpSupported())
 		{
 			return;
 		}
 
-		// Dump the resource's binary to a .bin file.
-		if (!ResourceVersionDumpName.IsEmpty())
+		FRHITexture* RHITexture = SubresourceDesc.Texture->GetRHI();
+
+		FShaderResourceViewRHIRef SubResourceSRV;
+		if (SubresourceDumpDesc.bPreprocessForStaging)
+		{
+			SubResourceSRV = RHICreateShaderResourceView(RHITexture, FRHITextureSRVCreateInfo(SubresourceDesc));
+			RHICmdListImmediate.Transition(FRHITransitionInfo(RHITexture, ERHIAccess::Unknown, ERHIAccess::SRVCompute));
+		}
+		else
 		{
-			FString DumpFilePath = ResourceDumpDirectory + ResourceVersionDumpName;
+			RHICmdListImmediate.Transition(FRHITransitionInfo(RHITexture, RHIAccessState, ERHIAccess::CopySrc));
+		}
 
-			FDumpTexturePass* PassParameters = GraphBuilder.AllocParameters<FDumpTexturePass>();
-			PassParameters->Texture = Texture;
+		FString DumpFilePath = kResourcesDir / FString::Printf(
+			TEXT("%s.v%016x.d%d.bin"),
+			*UniqueResourceSubResourceName,
+			PtrToUint(DrawDumpingPass),
+			DrawDumpCount);
+
+		DumpTextureSubResource(
+			RHICmdListImmediate,
+			SubresourceDesc.Texture->Name,
+			RHITexture,
+			SubResourceSRV,
+			SubresourceDumpDesc,
+			DumpFilePath);
+
+		if (SubresourceDumpDesc.bPreprocessForStaging)
+		{
+			RHICmdListImmediate.Transition(FRHITransitionInfo(RHITexture, ERHIAccess::SRVCompute, RHIAccessState));
+		}
+		else
+		{
+			RHICmdListImmediate.Transition(FRHITransitionInfo(RHITexture, ERHIAccess::CopySrc, RHIAccessState));
+		}
 
-			GraphBuilder.AddPass(
-				RDG_EVENT_NAME("RDG DumpTexture(%s -> %s)", Texture->Name, *DumpFilePath),
-				PassParameters,
-				ERDGPassFlags::Readback,
-				[DumpFilePath, Texture, ByteSize](FRHICommandListImmediate& RHICmdList)
-			{
-				check(IsInRenderingThread());
-				TUniquePtr<FRHIGPUTextureReadback> Readback = MakeUnique<FRHIGPUTextureReadback>(TEXT("RDGTextureDumpReadback"));;
+		SubResourceSRV = nullptr;
+		ReleaseRHIResources(RHICmdListImmediate);
+	}
 
-				Readback->EnqueueCopyRDG(RHICmdList, Texture->GetRHI(), FResolveRect());
+	void AddDumpTextureSubResourcePass(
+		FRDGBuilder& GraphBuilder,
+		TArray<TSharedPtr<FJsonValue>>& InputResourceNames,
+		TArray<TSharedPtr<FJsonValue>>& OutputResourceNames,
+		const FRDGPass* Pass,
+		FRDGTextureSRVDesc SubresourceDesc,
+		bool bIsOutputResource)
+	{
+		int32 DumpTextureMode = GDumpTextureCVar.GetValueOnRenderThread();
 
-				int32 RowPitchInPixels = 0;
-				void* Content = nullptr;
-				Readback->LockTexture(RHICmdList, /* out */ Content, /* out */ RowPitchInPixels);
-				if (Content)
-				{
-					TArray<uint8> Array;
-					Array.SetNumUninitialized(ByteSize);
+		if (DumpTextureMode == 0)
+		{
+			return;
+		}
+
+		const FRDGTextureDesc& Desc = SubresourceDesc.Texture->Desc;
+		const FString UniqueResourceName = GetUniqueResourceName(SubresourceDesc.Texture);
+		const FString UniqueResourceSubResourceName = GetUniqueSubResourceName(SubresourceDesc);
+		const FTextureSubresourceDumpDesc SubresourceDumpDesc = TranslateSubresourceDumpDesc(SubresourceDesc);
+
+		if (bIsOutputResource)
+		{
+			OutputResourceNames.AddUnique(MakeShareable(new FJsonValueString(UniqueResourceSubResourceName)));
+		}
+		else
+		{
+			InputResourceNames.AddUnique(MakeShareable(new FJsonValueString(UniqueResourceSubResourceName)));
+		}
+
+		bool bDumpResourceInfos;
+		bool bDumpResourceBinary;
+		GetResourceDumpInfo(Pass, SubresourceDesc.Texture, bIsOutputResource, &bDumpResourceInfos, &bDumpResourceBinary);
+
+		// Dump the information of the texture to json file.
+		if (bDumpResourceInfos)
+		{
+			TSharedPtr<FJsonObject> JsonObject = ToJson(SubresourceDesc.Texture->Name, Desc);
+
+			FString ResourceInfoJsonPath = kResourcesDir / UniqueResourceName + TEXT(".json");
+			DumpJsonToFile(JsonObject, ResourceInfoJsonPath);
+		}
 
-					int32 BytePerPixel = GPixelFormats[Texture->Desc.Format].BlockBytes;
+		if (!SubresourceDumpDesc.IsDumpSupported())
+		{
+			return;
+		}
 
-					const uint8* SrcData = reinterpret_cast<const uint8*>(Content);
+		// Early return if this resource shouldn't be dumped.
+		if (!bDumpResourceBinary || DumpTextureMode != 2)
+		{
+			return;
+		}
 
-					for (int32 y = 0; y < Texture->Desc.Extent.Y; y++)
-					{
-						// Flip the data to be bottom left corner for the WebGL viewer.
-						const uint8* SrcPos = SrcData + (Texture->Desc.Extent.Y - 1 - y) * RowPitchInPixels * BytePerPixel;
-						uint8* DstPos = (&Array[0]) + y * Texture->Desc.Extent.X * BytePerPixel;
+		FString DumpFilePath = kResourcesDir / FString::Printf(TEXT("%s.v%016x.bin"), *UniqueResourceSubResourceName, PtrToUint(bIsOutputResource ? Pass : nullptr));
 
-						FPlatformMemory::Memmove(DstPos, SrcPos, Texture->Desc.Extent.X * BytePerPixel);
-					}
+		// Verify there is enough available memory to dump the resource.
+		if (IsUnsafeToDumpResource(SubresourceDumpDesc.ByteSize, 2.2f + (SubresourceDumpDesc.bPreprocessForStaging ? 1.0f : 0.0f)))
+		{
+			UE_LOG(LogRendererCore, Warning, TEXT("Not dumping %s because of insuficient memory available for staging texture."), *DumpFilePath);
+			return;
+		}
 
-					TArrayView<const uint8> ArrayView(reinterpret_cast<const uint8*>(&Array[0]), ByteSize);
-					FFileHelper::SaveArrayToFile(ArrayView, *DumpFilePath);
+		// Dump the resource's binary to a .bin file.
+		{
+			FDumpTexturePass* PassParameters = GraphBuilder.AllocParameters<FDumpTexturePass>();
+			if (SubresourceDumpDesc.bPreprocessForStaging)
+			{
+				if (!(SubresourceDesc.Texture->Desc.Flags & TexCreate_ShaderResource))
+				{
+					UE_LOG(LogRendererCore, Warning, TEXT("Not dumping %s because requires copy to staging texture using compute, but is missing TexCreate_ShaderResource."), *UniqueResourceSubResourceName);
+					return;
+				}
 
-					Readback->Unlock();
+				if (!FDumpTextureCS::IsSupported(GMaxRHIShaderPlatform))
+				{
+					UE_LOG(LogRendererCore, Warning, TEXT("Not dumping %s because FDumpTextureCS compute shader is not supported."), *UniqueResourceSubResourceName);
+					return;
 				}
+
+				PassParameters->Texture = GraphBuilder.CreateSRV(SubresourceDesc);
+			}
+			else
+			{
+				PassParameters->TextureAccess = FRDGTextureAccess(SubresourceDesc.Texture, ERHIAccess::CopySrc);
+			}
+
+			GraphBuilder.AddPass(
+				RDG_EVENT_NAME("RDG DumpTexture(%s -> %s) %dx%d",
+					SubresourceDesc.Texture->Name, *DumpFilePath,
+					SubresourceDumpDesc.SubResourceExtent.X, SubresourceDumpDesc.SubResourceExtent.Y),
+				PassParameters,
+				(SubresourceDumpDesc.bPreprocessForStaging ? ERDGPassFlags::Compute : ERDGPassFlags::Copy) | ERDGPassFlags::NeverCull,
+				[
+					PassParameters, this, DumpFilePath,
+					SubresourceDesc, SubresourceDumpDesc]
+				(FRHICommandListImmediate& RHICmdList)
+			{
+
+				this->DumpTextureSubResource(
+					RHICmdList,
+					SubresourceDesc.Texture->Name,
+					SubresourceDumpDesc.bPreprocessForStaging ? nullptr : SubresourceDesc.Texture->GetRHI(),
+					SubresourceDumpDesc.bPreprocessForStaging ? PassParameters->Texture->GetRHI() : nullptr,
+					SubresourceDumpDesc,
+					DumpFilePath);
+				this->ReleaseRHIResources(RHICmdList);
+				this->UpdatePassProgress();
 			});
+
+			ResourcesDumpPasses++;
 		}
+	}
 
-		ResourcesDumped++;
+	void AddDumpTexturePasses(
+		FRDGBuilder& GraphBuilder,
+		TArray<TSharedPtr<FJsonValue>>& InputResourceNames,
+		TArray<TSharedPtr<FJsonValue>>& OutputResourceNames,
+		const FRDGPass* Pass,
+		FRDGTextureSRVDesc SubresourceRangeDesc,
+		bool bIsOutputResource)
+	{
+		if (SubresourceRangeDesc.Format == PF_X24_G8)
+		{
+			AddDumpTextureSubResourcePass(
+				GraphBuilder,
+				InputResourceNames,
+				OutputResourceNames,
+				Pass,
+				SubresourceRangeDesc,
+				bIsOutputResource);
+		}
+		else
+		{
+			for (int32 MipLevel = SubresourceRangeDesc.MipLevel; MipLevel < (SubresourceRangeDesc.MipLevel + SubresourceRangeDesc.NumMipLevels); MipLevel++)
+			{
+				FRDGTextureSRVDesc SubresourceDesc = FRDGTextureSRVDesc::CreateForMipLevel(SubresourceRangeDesc.Texture, MipLevel);
+				AddDumpTextureSubResourcePass(
+					GraphBuilder,
+					InputResourceNames,
+					OutputResourceNames,
+					Pass,
+					SubresourceDesc,
+					bIsOutputResource);
+			}
+		}
 	}
 
 	void AddDumpBufferPass(
@@ -233,7 +996,14 @@ struct FRDGResourceDumpContext
 		FRDGBuffer* Buffer,
 		bool bIsOutputResource)
 	{
-		FString UniqueResourceName = GetResourceName(Buffer);
+		int32 DumpBufferMode = GDumpTextureCVar.GetValueOnRenderThread();
+
+		if (DumpBufferMode == 0)
+		{
+			return;
+		}
+
+		FString UniqueResourceName = GetUniqueResourceName(Buffer);
 
 		if (bIsOutputResource)
 		{
@@ -245,30 +1015,24 @@ struct FRDGResourceDumpContext
 		}
 
 		const FRDGBufferDesc& Desc = Buffer->Desc;
-		int32 ByteSize = Desc.GetTotalNumBytes();
-
-		FString ResourceDumpDirectory = GetResourceDumpDirectory();
+		const int32 ByteSize = Desc.GetTotalNumBytes();
 
 		bool bDumpResourceInfos;
-		FString ResourceVersionDumpName;
-		GetResourceDumpInfo(Pass, Buffer, bIsOutputResource, &bDumpResourceInfos, &ResourceVersionDumpName);
+		bool bDumpResourceBinary;
+		GetResourceDumpInfo(Pass, Buffer, bIsOutputResource, &bDumpResourceInfos, &bDumpResourceBinary);
 
 		// Dump the information of the buffer to json file.
 		if (bDumpResourceInfos)
 		{
-			TSharedPtr<FJsonObject> JsonObject = MakeShareable(new FJsonObject);
-			JsonObject->SetStringField(TEXT("Name"), Buffer->Name);
-			JsonObject->SetNumberField(TEXT("ByteSize"), ByteSize);
-			JsonObject->SetStringField(TEXT("Type"), GetBufferUnderlyingTypeName(Desc.UnderlyingType));
-			JsonObject->SetNumberField(TEXT("BytesPerElement"), Desc.BytesPerElement);
-			JsonObject->SetNumberField(TEXT("NumElements"), Desc.NumElements);
+			TSharedPtr<FJsonObject> JsonObject = ToJson(Buffer->Name, Desc, ByteSize);
 
-			FString OutputString;
-			TSharedRef<TJsonWriter<TCHAR, TPrettyJsonPrintPolicy<TCHAR>>> Writer = TJsonWriterFactory<TCHAR, TPrettyJsonPrintPolicy<TCHAR>>::Create(&OutputString);
-			FJsonSerializer::Serialize(JsonObject.ToSharedRef(), Writer);
+			FString ResourceInfoJsonPath = kResourcesDir / UniqueResourceName + TEXT(".json");
+			DumpJsonToFile(JsonObject, ResourceInfoJsonPath);
 
-			FString ResourceInfoJsonPath = ResourceDumpDirectory + UniqueResourceName + TEXT(".json");
-			FFileHelper::SaveStringToFile(OutputString, *ResourceInfoJsonPath);
+			if (Desc.Metadata && DumpBufferMode == 2)
+			{
+				Dump(Desc.Metadata);
+			}
 		}
 
 		if (Desc.UnderlyingType == FRDGBufferDesc::EUnderlyingType::AccelerationStructure)
@@ -277,9 +1041,15 @@ struct FRDGResourceDumpContext
 		}
 
 		// Dump the resource's binary to a .bin file.
-		if (!ResourceVersionDumpName.IsEmpty())
+		if (bDumpResourceBinary && DumpBufferMode == 2)
 		{
-			FString DumpFilePath = ResourceDumpDirectory + ResourceVersionDumpName;
+			FString DumpFilePath = kResourcesDir / FString::Printf(TEXT("%s.v%016x.bin"), *UniqueResourceName, PtrToUint(bIsOutputResource ? Pass : nullptr));
+
+			if (IsUnsafeToDumpResource(ByteSize, 1.2f))
+			{
+				UE_LOG(LogRendererCore, Warning, TEXT("Not dumping %s because of insuficient memory available for staging buffer."), *DumpFilePath);
+				return;
+			}
 
 			FDumpBufferPass* PassParameters = GraphBuilder.AllocParameters<FDumpBufferPass>();
 			PassParameters->Buffer = Buffer;
@@ -288,127 +1058,241 @@ struct FRDGResourceDumpContext
 				RDG_EVENT_NAME("RDG DumpBuffer(%s -> %s)", Buffer->Name, *DumpFilePath),
 				PassParameters,
 				ERDGPassFlags::Readback,
-				[DumpFilePath, Buffer, ByteSize](FRHICommandListImmediate& RHICmdList)
+				[this, DumpFilePath, Buffer, ByteSize](FRHICommandListImmediate& RHICmdList)
 			{
 				check(IsInRenderingThread());
-				TUniquePtr<FRHIGPUBufferReadback> Readback = MakeUnique<FRHIGPUBufferReadback>(TEXT("RDGBufferDumpReadback"));;
+				FStagingBufferRHIRef StagingBuffer = RHICreateStagingBuffer();
+
+				// Transfer memory GPU -> CPU
+				RHICmdList.CopyToStagingBuffer(Buffer->GetRHI(), StagingBuffer, 0, ByteSize);
+
+				// Submit to GPU and wait for completion.
+				static const FName FenceName(TEXT("DumpGPU.BufferFence"));
+				FGPUFenceRHIRef Fence = RHICreateGPUFence(FenceName);
+				{
+					Fence->Clear();
+					RHICmdList.WriteGPUFence(Fence);
+					RHICmdList.SubmitCommandsAndFlushGPU();
+					RHICmdList.BlockUntilGPUIdle();
+				}
+
+				void* Content = RHICmdList.LockStagingBuffer(StagingBuffer, Fence.GetReference(), 0, ByteSize);
+				if (Content)
+				{
+					TArrayView<const uint8> ArrayView(reinterpret_cast<const uint8*>(Content), ByteSize);
+					DumpBinaryToFile(ArrayView, DumpFilePath);
+
+					RHICmdList.UnlockStagingBuffer(StagingBuffer);
+				}
+				else
+				{
+					UE_LOG(LogRendererCore, Warning, TEXT("RHICmdList.MapStagingSurface() to dump buffer %s failed."), Buffer->Name);
+				}
+
+				StagingBuffer = nullptr;
+				Fence = nullptr;
+				this->ReleaseRHIResources(RHICmdList);
+				this->UpdatePassProgress();
+			});
+
+			ResourcesDumpPasses++;
+		}
+	}
+
+	// Look whether the pass matches matches r.DumpGPU.Root
+	bool IsDumpingPass(const FRDGPass* Pass)
+	{
+		FString RootWildcardString = GDumpGPURootCVar.GetValueOnRenderThread();
+		FWildcardString WildcardFilter(RootWildcardString);
 
-				// TODO: work arround BUF_SourceCopy ensure()
-				Readback->EnqueueCopy(RHICmdList, Buffer->GetRHI(), ByteSize);
+		bool bDumpPass = (RootWildcardString == TEXT("*"));
 
-				void* Content = Readback->Lock(ByteSize);
-				if (Content)
-				{
-					TArrayView<const uint8> ArrayView(reinterpret_cast<const uint8*>(Content), ByteSize);
-					FFileHelper::SaveArrayToFile(ArrayView, *DumpFilePath);
+		if (!bDumpPass)
+		{
+			bDumpPass = WildcardFilter.IsMatch(Pass->GetEventName().GetTCHAR());
+		}
 
-					Readback->Unlock();
-				}
-				else
-				{
-					ensureMsgf(false, TEXT("%s couldn't be dumped"), Buffer->Name);
-				}
-			});
+		#if RDG_GPU_SCOPES
+		if (!bDumpPass)
+		{
+			const FRDGEventScope* ParentScope = Pass->GetGPUScopes().Event;
+
+			while (ParentScope)
+			{
+				bDumpPass = bDumpPass || WildcardFilter.IsMatch(ParentScope->Name.GetTCHAR());
+				ParentScope = ParentScope->ParentScope;
+			}
 		}
+		#endif
 
-		ResourcesDumped++;
+		return bDumpPass;
 	}
 };
 
+bool bIsDumpingFrame_GameThread = false;
 FRDGResourceDumpContext GRDGResourceDumpContext;
 
-static FAutoConsoleCommand GDumpGPUCommand(
-	TEXT("r.DumpGPU"),
-	TEXT("Dump a multi-platform GPU capture of RDG resources, to be used when creating a JIRA of an artifact/rendering bug."),
-	FConsoleCommandWithArgsDelegate::CreateStatic(
-		[](const TArray<FString>& Args)
-		{
-			check(IsInGameThread());
-
-			ENQUEUE_RENDER_COMMAND(FSceneRendererCleanUp)(
-				[Args](FRHICommandListImmediate& RHICmdList)
-			{
-				FRDGBuilder::BeginResourceDump(Args);
-			});
-		}
-	)
-);
+}
 
+bool IsDumpingRDGResources()
+{
+	return GRDGResourceDumpContext.IsDumpingFrame();
 }
 
-void FRDGBuilder::BeginResourceDump(const TArray<FString>& Args)
+FString FRDGBuilder::BeginResourceDump(const TArray<FString>& Args)
 {
-	if (GRDGResourceDumpContext.IsDumpingFrame())
+	check(IsInGameThread());
+
+	if (bIsDumpingFrame_GameThread)
 	{
-		return;
+		return FString();
 	}
 
-	GRDGResourceDumpContext.Time = FDateTime::Now();
-	GRDGResourceDumpContext.DumpingDirectoryPath = FPaths::ProjectSavedDir() + TEXT("GPUDumps/") + FApp::GetProjectName() + TEXT("-") + FPlatformProperties::PlatformName() + TEXT("-") + GRDGResourceDumpContext.Time.ToString() + TEXT("/");
+	FRDGResourceDumpContext NewResourceDumpContext;
 
-	GRDGResourceDumpContext.bShowInExplore = Args.Contains(TEXT("-show"));
+	NewResourceDumpContext.Time = FDateTime::Now();
+	{
+		FString MasterDirectoryPath = GDumpGPUDirectoryCVar.GetValueOnGameThread();
+		if (MasterDirectoryPath.IsEmpty())
+		{
+			MasterDirectoryPath = FPaths::ProjectSavedDir() / TEXT("GPUDumps/");
+		}
+		NewResourceDumpContext.DumpingDirectoryPath = MasterDirectoryPath / FApp::GetProjectName() + TEXT("-") + FPlatformProperties::PlatformName() + TEXT("-") + NewResourceDumpContext.Time.ToString() + TEXT("/");
+	}
+	NewResourceDumpContext.bEnableDiskWrite = GDumpTestEnableDiskWrite.GetValueOnGameThread() != 0;
+	NewResourceDumpContext.bShowInExplore = NewResourceDumpContext.bEnableDiskWrite && GDumpExploreCVar.GetValueOnGameThread() != 0;
+	NewResourceDumpContext.MemoryConstants = FPlatformMemory::GetConstants();
+	NewResourceDumpContext.MemoryStats = FPlatformMemory::GetStats();
 
 	IPlatformFile& PlatformFile = FPlatformFileManager::Get().GetPlatformFile();
-	if (!PlatformFile.DirectoryExists(*GRDGResourceDumpContext.DumpingDirectoryPath))
+	if (NewResourceDumpContext.bEnableDiskWrite)
 	{
-		PlatformFile.CreateDirectoryTree(*GRDGResourceDumpContext.DumpingDirectoryPath);
+		if (!PlatformFile.DirectoryExists(*NewResourceDumpContext.DumpingDirectoryPath))
+		{
+			PlatformFile.CreateDirectoryTree(*NewResourceDumpContext.DumpingDirectoryPath);
+		}
+		PlatformFile.CreateDirectoryTree(*(NewResourceDumpContext.DumpingDirectoryPath / FRDGResourceDumpContext::kBaseDir));
+		PlatformFile.CreateDirectoryTree(*(NewResourceDumpContext.DumpingDirectoryPath / FRDGResourceDumpContext::kResourcesDir));
+
+		NewResourceDumpContext.DumpStringToFile(TEXT(""), FString(FRDGResourceDumpContext::kBaseDir) / TEXT("PassDrawCounts.json"));
 	}
-	PlatformFile.CreateDirectoryTree(*(GRDGResourceDumpContext.DumpingDirectoryPath + TEXT("Resources/")));
 
 	// Output informations
 	{
+		const TCHAR* BranchName = BuildSettings::GetBranchName();
+		const TCHAR* BuildDate = BuildSettings::GetBuildDate();
+		const TCHAR* BuildVersion = BuildSettings::GetBuildVersion();
+
+		FString BuildConfiguration = LexToString(FApp::GetBuildConfiguration());
+		FString BuildTarget = LexToString(FApp::GetBuildTargetType());
+
+		FGPUDriverInfo GPUDriverInfo = FPlatformMisc::GetGPUDriverInfo(GRHIAdapterName);
+
 		TSharedPtr<FJsonObject> JsonObject = MakeShareable(new FJsonObject);
 		JsonObject->SetStringField(TEXT("Project"), FApp::GetProjectName());
-		JsonObject->SetStringField(TEXT("BuildBranch"), BuildSettings::GetBranchName());
-		JsonObject->SetStringField(TEXT("BuildDate"), BuildSettings::GetBuildDate());
-		JsonObject->SetStringField(TEXT("BuildVersion"), BuildSettings::GetBuildVersion());
-		JsonObject->SetStringField(TEXT("DumpTime"), GRDGResourceDumpContext.Time.ToString());
+		JsonObject->SetNumberField(TEXT("EngineMajorVersion"), ENGINE_MAJOR_VERSION);
+		JsonObject->SetNumberField(TEXT("EngineMinorVersion"), ENGINE_MINOR_VERSION);
+		JsonObject->SetNumberField(TEXT("EnginePatchVersion"), ENGINE_PATCH_VERSION);
+		JsonObject->SetStringField(TEXT("BuildBranch"), BranchName ? BranchName : TEXT(""));
+		JsonObject->SetStringField(TEXT("BuildDate"), BuildDate ? BuildDate : TEXT(""));
+		JsonObject->SetStringField(TEXT("BuildVersion"), BuildVersion ? BuildVersion : TEXT(""));
+		JsonObject->SetStringField(TEXT("BuildTarget"), BuildTarget);
+		JsonObject->SetStringField(TEXT("BuildConfiguration"), BuildConfiguration);
+		JsonObject->SetNumberField(TEXT("Build64Bits"), (PLATFORM_64BITS ? 1 : 0));
 		JsonObject->SetStringField(TEXT("Platform"), FPlatformProperties::IniPlatformName());
-		JsonObject->SetStringField(TEXT("RHI"), FApp::GetGraphicsRHI());
+		JsonObject->SetStringField(TEXT("DeviceName"), FPlatformProcess::ComputerName());
+		JsonObject->SetStringField(TEXT("CPUVendor"), FPlatformMisc::GetCPUVendor());
+		JsonObject->SetStringField(TEXT("CPUBrand"), FPlatformMisc::GetCPUBrand());
+		JsonObject->SetNumberField(TEXT("CPUNumberOfCores"), FPlatformMisc::NumberOfCores());
+		JsonObject->SetNumberField(TEXT("CPUNumberOfCoresIncludingHyperthreads"), FPlatformMisc::NumberOfCoresIncludingHyperthreads());
 		JsonObject->SetStringField(TEXT("GPUVendor"), RHIVendorIdToString());
-		JsonObject->SetStringField(TEXT("ComputerName"), FPlatformProcess::ComputerName());
-
-		FString OutputString;
-		TSharedRef<TJsonWriter<TCHAR, TPrettyJsonPrintPolicy<TCHAR>>> Writer = TJsonWriterFactory<TCHAR, TPrettyJsonPrintPolicy<TCHAR>>::Create(&OutputString);
-		FJsonSerializer::Serialize(JsonObject.ToSharedRef(), Writer);
+		JsonObject->SetStringField(TEXT("GPUDeviceDescription"), GPUDriverInfo.DeviceDescription);
+		JsonObject->SetStringField(TEXT("GPUDriverUserVersion"), GPUDriverInfo.UserDriverVersion);
+		JsonObject->SetStringField(TEXT("GPUDriverInternalVersion"), GPUDriverInfo.GetUnifiedDriverVersion());
+		JsonObject->SetStringField(TEXT("GPUDriverDate"), GPUDriverInfo.DriverDate);
+		JsonObject->SetNumberField(TEXT("MemoryTotalPhysical"), NewResourceDumpContext.MemoryConstants.TotalPhysical);
+		JsonObject->SetNumberField(TEXT("MemoryPageSize"), NewResourceDumpContext.MemoryConstants.PageSize);
+		JsonObject->SetStringField(TEXT("RHI"), GDynamicRHI->GetName());
+		JsonObject->SetStringField(TEXT("RHIMaxFeatureLevel"), LexToString(GMaxRHIFeatureLevel));
+		JsonObject->SetStringField(TEXT("DumpTime"), NewResourceDumpContext.Time.ToString());
+
+		NewResourceDumpContext.DumpJsonToFile(JsonObject, FString(FRDGResourceDumpContext::kBaseDir) / TEXT("Infos.json"));
+	}
 
-		FString ResourceInfoJsonPath = GRDGResourceDumpContext.DumpingDirectoryPath + TEXT("Infos.json");
-		FFileHelper::SaveStringToFile(OutputString, *ResourceInfoJsonPath);
+	// Dump the rendering cvars
+	if (NewResourceDumpContext.bEnableDiskWrite)
+	{
+		NewResourceDumpContext.DumpRenderingCVarsToCSV();
 	}
 
 	// Copy the viewer
+	if (NewResourceDumpContext.bEnableDiskWrite)
 	{
 		const TCHAR* OpenGPUDumpViewerWindowsName = TEXT("OpenGPUDumpViewer.bat");
 		const TCHAR* ViewerHTML = TEXT("GPUDumpViewer.html");
 		FString DumpGPUViewerSourcePath = FPaths::EngineDir() + FString(TEXT("Extras")) / TEXT("GPUDumpViewer");
 
-		PlatformFile.CopyFile(*(GRDGResourceDumpContext.DumpingDirectoryPath / ViewerHTML), *(DumpGPUViewerSourcePath / ViewerHTML));
-		PlatformFile.CopyFile(*(GRDGResourceDumpContext.DumpingDirectoryPath / OpenGPUDumpViewerWindowsName), *(DumpGPUViewerSourcePath / OpenGPUDumpViewerWindowsName));
+		PlatformFile.CopyFile(*(NewResourceDumpContext.DumpingDirectoryPath / ViewerHTML), *(DumpGPUViewerSourcePath / ViewerHTML));
+		PlatformFile.CopyFile(*(NewResourceDumpContext.DumpingDirectoryPath / OpenGPUDumpViewerWindowsName), *(DumpGPUViewerSourcePath / OpenGPUDumpViewerWindowsName));
+	}
+
+	ENQUEUE_RENDER_COMMAND(FStartGPUDump)(
+		[NewResourceDumpContext](FRHICommandListImmediate& ImmediateRHICmdList)
+	{
+		check(IsInRenderingThread());
+		GRDGResourceDumpContext = NewResourceDumpContext;
+	});
+
+	bIsDumpingFrame_GameThread = true;
+
+	if (NewResourceDumpContext.bEnableDiskWrite)
+	{
+		return NewResourceDumpContext.DumpingDirectoryPath;
 	}
+	return FString();
 }
 
 void FRDGBuilder::EndResourceDump()
 {
-	if (!GRDGResourceDumpContext.IsDumpingFrame())
+	check(IsInGameThread());
+
+	if (!bIsDumpingFrame_GameThread)
 	{
 		return;
 	}
 
+	// Wait all rendering commands are completed to finish with GRDGResourceDumpContext.
+	{
+		UE_LOG(LogRendererCore, Display, TEXT("Stalling game thread until render thread finishes to dump resources"));
+		FlushRenderingCommands();
+	}
+
+	// Log information about the dump.
 	FString AbsDumpingDirectoryPath = IFileManager::Get().ConvertToAbsolutePathForExternalAppForRead(*GRDGResourceDumpContext.DumpingDirectoryPath);
-	UE_LOG(LogRendererCore, Display, TEXT("Dumped %d resources to %s"), GRDGResourceDumpContext.ResourcesDumped, *AbsDumpingDirectoryPath);
+	{
+		UE_LOG(LogRendererCore, Display, TEXT("Dumped %d resources to %s"), GRDGResourceDumpContext.ResourcesDumpPasses, *AbsDumpingDirectoryPath);
+	}
 
-	if (GRDGResourceDumpContext.bShowInExplore)
-	#if PLATFORM_DESKTOP
+	// Dump the log into the dump directory.
+	if (GRDGResourceDumpContext.bEnableDiskWrite)
 	{
-		FPlatformProcess::ExploreFolder(*AbsDumpingDirectoryPath);
+		if (GLog)
+		{
+			GLog->FlushThreadedLogs();
+			GLog->Flush();
+		}
+		FGenericCrashContext::DumpLog(GRDGResourceDumpContext.DumpingDirectoryPath / FRDGResourceDumpContext::kBaseDir);
 	}
-	#else
+
+	#if PLATFORM_DESKTOP
+	if (GRDGResourceDumpContext.bShowInExplore)
 	{
-		UE_LOG(LogRendererCore, Warning, TEXT("r.DumpGPU parameter -show is not supported on this platform."));
+		FPlatformProcess::ExploreFolder(*AbsDumpingDirectoryPath);
 	}
 	#endif
 
 	GRDGResourceDumpContext = FRDGResourceDumpContext();
+	bIsDumpingFrame_GameThread = false;
 }
 
 void FRDGBuilder::DumpResourcePassOutputs(const FRDGPass* Pass)
@@ -423,35 +1307,10 @@ void FRDGBuilder::DumpResourcePassOutputs(const FRDGPass* Pass)
 		return;
 	}
 
-	// Look whether the pass matches matches r.DumpGPU.Root
+	check(IsInRenderingThread());
+	if (!GRDGResourceDumpContext.IsDumpingPass(Pass))
 	{
-		FString RootWildcardString = GDumpGPURootCVar.GetValueOnRenderThread();
-		FWildcardString WildcardFilter(RootWildcardString);
-
-		bool bDumpPass = (RootWildcardString == TEXT("*"));
-
-		if (!bDumpPass)
-		{
-			bDumpPass = WildcardFilter.IsMatch(Pass->GetEventName().GetTCHAR());
-		}
-		
-		#if RDG_GPU_SCOPES
-		if (!bDumpPass)
-		{
-			const FRDGEventScope* ParentScope = Pass->GetGPUScopes().Event;
-
-			while (ParentScope)
-			{
-				bDumpPass = bDumpPass || WildcardFilter.IsMatch(ParentScope->Name.GetTCHAR());
-				ParentScope = ParentScope->ParentScope;
-			}
-		}
-		#endif
-
-		if (!bDumpPass)
-		{
-			return;
-		}
+		return;
 	}
 
 	bInDebugPassScope = true;
@@ -466,7 +1325,8 @@ void FRDGBuilder::DumpResourcePassOutputs(const FRDGPass* Pass)
 		{
 			if (FRDGTextureRef Texture = Parameter.GetAsTexture())
 			{
-				GRDGResourceDumpContext.AddDumpTexturePass(*this, InputResourceNames, OutputResourceNames, Pass, Texture, /* bIsOutputResource = */ false);
+				FRDGTextureSRVDesc TextureSubResource = FRDGTextureSRVDesc::Create(Texture);
+				GRDGResourceDumpContext.AddDumpTexturePasses(*this, InputResourceNames, OutputResourceNames, Pass, TextureSubResource, /* bIsOutputResource = */ false);
 			}
 		}
 		break;
@@ -474,8 +1334,14 @@ void FRDGBuilder::DumpResourcePassOutputs(const FRDGPass* Pass)
 		{
 			if (FRDGTextureSRVRef SRV = Parameter.GetAsTextureSRV())
 			{
-				FRDGTextureRef Texture = SRV->Desc.Texture;
-				GRDGResourceDumpContext.AddDumpTexturePass(*this, InputResourceNames, OutputResourceNames, Pass, Texture, /* bIsOutputResource = */ false);
+				if (SRV->Desc.MetaData == ERHITextureMetaDataAccess::None)
+				{
+					GRDGResourceDumpContext.AddDumpTexturePasses(*this, InputResourceNames, OutputResourceNames, Pass, SRV->Desc, /* bIsOutputResource = */ false);
+				}
+				else
+				{
+					UE_LOG(LogRendererCore, Warning, TEXT("Dumping texture %s's meta data unsupported"), SRV->Desc.Texture->Name);
+				}
 			}
 		}
 		break;
@@ -483,8 +1349,15 @@ void FRDGBuilder::DumpResourcePassOutputs(const FRDGPass* Pass)
 		{
 			if (FRDGTextureUAVRef UAV = Parameter.GetAsTextureUAV())
 			{
-				FRDGTextureRef Texture = UAV->Desc.Texture;
-				GRDGResourceDumpContext.AddDumpTexturePass(*this, InputResourceNames, OutputResourceNames, Pass, Texture, /* bIsOutputResource = */ true);
+				if (UAV->Desc.MetaData == ERHITextureMetaDataAccess::None)
+				{
+					FRDGTextureSRVDesc TextureSubResource = FRDGTextureSRVDesc::CreateForMipLevel(UAV->Desc.Texture, UAV->Desc.MipLevel);
+					GRDGResourceDumpContext.AddDumpTextureSubResourcePass(*this, InputResourceNames, OutputResourceNames, Pass, TextureSubResource, /* bIsOutputResource = */ true);
+				}
+				else
+				{
+					UE_LOG(LogRendererCore, Warning, TEXT("Dumping texture %s's meta data unsupported"), UAV->Desc.Texture->Name);
+				}
 			}
 		}
 		break;
@@ -497,7 +1370,8 @@ void FRDGBuilder::DumpResourcePassOutputs(const FRDGPass* Pass)
 					TextureAccess.GetAccess() == ERHIAccess::UAVGraphics ||
 					TextureAccess.GetAccess() == ERHIAccess::RTV);
 
-				GRDGResourceDumpContext.AddDumpTexturePass(*this, InputResourceNames, OutputResourceNames, Pass, TextureAccess, bIsOutputResource);
+				FRDGTextureSRVDesc TextureSubResource = FRDGTextureSRVDesc::Create(TextureAccess);
+				GRDGResourceDumpContext.AddDumpTexturePasses(*this, InputResourceNames, OutputResourceNames, Pass, TextureSubResource, bIsOutputResource);
 			}
 		}
 		break;
@@ -512,7 +1386,8 @@ void FRDGBuilder::DumpResourcePassOutputs(const FRDGPass* Pass)
 					TextureAccess.GetAccess() == ERHIAccess::UAVGraphics ||
 					TextureAccess.GetAccess() == ERHIAccess::RTV);
 
-				GRDGResourceDumpContext.AddDumpTexturePass(*this, InputResourceNames, OutputResourceNames, Pass, TextureAccess, bIsOutputResource);
+				FRDGTextureSRVDesc TextureSubResource = FRDGTextureSRVDesc::Create(TextureAccess);
+				GRDGResourceDumpContext.AddDumpTexturePasses(*this, InputResourceNames, OutputResourceNames, Pass, TextureSubResource, bIsOutputResource);
 			}
 		}
 		break;
@@ -569,21 +1444,38 @@ void FRDGBuilder::DumpResourcePassOutputs(const FRDGPass* Pass)
 			RenderTargets.Enumerate([&](FRenderTargetBinding RenderTarget)
 			{
 				FRDGTextureRef Texture = RenderTarget.GetTexture();
-				GRDGResourceDumpContext.AddDumpTexturePass(*this, InputResourceNames, OutputResourceNames, Pass, Texture, /* bIsOutputResource = */ true);
+				FRDGTextureSRVDesc TextureSubResource = FRDGTextureSRVDesc::CreateForMipLevel(Texture, RenderTarget.GetMipIndex());
+				GRDGResourceDumpContext.AddDumpTextureSubResourcePass(*this, InputResourceNames, OutputResourceNames, Pass, TextureSubResource, /* bIsOutputResource = */ true);
 			});
 
 			const FDepthStencilBinding& DepthStencil = RenderTargets.DepthStencil;
 
 			if (FRDGTextureRef Texture = DepthStencil.GetTexture())
 			{
-				const bool bHasStoreAction = DepthStencil.GetDepthStencilAccess().IsAnyWrite();
-				GRDGResourceDumpContext.AddDumpTexturePass(*this, InputResourceNames, OutputResourceNames, Pass, Texture, /* bIsOutputResource = */ bHasStoreAction);
+				FExclusiveDepthStencil DepthStencilAccess = DepthStencil.GetDepthStencilAccess();
+
+				if (DepthStencilAccess.IsUsingDepth())
+				{
+					FRDGTextureSRVDesc TextureSubResource = FRDGTextureSRVDesc::CreateForMipLevel(Texture, 0);
+					GRDGResourceDumpContext.AddDumpTextureSubResourcePass(
+						*this, InputResourceNames, OutputResourceNames, Pass, TextureSubResource,
+						/* bIsOutputResource = */ DepthStencilAccess.IsDepthWrite());
+				}
+
+				if (DepthStencilAccess.IsUsingStencil())
+				{
+					FRDGTextureSRVDesc TextureSubResource = FRDGTextureSRVDesc::CreateWithPixelFormat(Texture, PF_X24_G8);
+					GRDGResourceDumpContext.AddDumpTextureSubResourcePass(
+						*this, InputResourceNames, OutputResourceNames, Pass, TextureSubResource,
+						/* bIsOutputResource = */ DepthStencilAccess.IsStencilWrite());
+				}
 			}
 		}
 		break;
 		}
 	});
 
+	// Dump the pass informations
 	{
 		TArray<TSharedPtr<FJsonValue>> ParentEventScopeNames;
 		#if RDG_GPU_SCOPES
@@ -597,24 +1489,50 @@ void FRDGBuilder::DumpResourcePassOutputs(const FRDGPass* Pass)
 			}
 		}
 		#endif
+		{
+			ParentEventScopeNames.Add(MakeShareable(new FJsonValueString(FString::Printf(TEXT("Frame %llu"), GFrameCounterRenderThread))));
+		}
 
 		TSharedPtr<FJsonObject> JsonObject = MakeShareable(new FJsonObject);
 		JsonObject->SetStringField(TEXT("EventName"), Pass->GetEventName().GetTCHAR());
 		JsonObject->SetStringField(TEXT("ParametersName"), Pass->GetParameters().GetLayout().GetDebugName());
-		JsonObject->SetStringField(TEXT("Pointer"), FString::Printf(TEXT("%p"), Pass));
+		JsonObject->SetStringField(TEXT("Parameters"), FRDGResourceDumpContext::PtrToString(Pass->GetParameters().GetContents()));
+		JsonObject->SetStringField(TEXT("ParametersMetadata"), FRDGResourceDumpContext::PtrToString(Pass->GetParameters().GetMetadata()));
+		JsonObject->SetStringField(TEXT("Pointer"), FString::Printf(TEXT("%016x"), FRDGResourceDumpContext::PtrToUint(Pass)));
 		JsonObject->SetNumberField(TEXT("Id"), GRDGResourceDumpContext.PassesCount);
 		JsonObject->SetArrayField(TEXT("ParentEventScopes"), ParentEventScopeNames);
 		JsonObject->SetArrayField(TEXT("InputResources"), InputResourceNames);
 		JsonObject->SetArrayField(TEXT("OutputResources"), OutputResourceNames);
 
-		FString OutputString;
-		TSharedRef<TJsonWriter<TCHAR, TPrettyJsonPrintPolicy<TCHAR>>> Writer = TJsonWriterFactory<TCHAR, TPrettyJsonPrintPolicy<TCHAR>>::Create(&OutputString);
-		FJsonSerializer::Serialize(JsonObject.ToSharedRef(), Writer);
+		GRDGResourceDumpContext.DumpJsonToFile(JsonObject, FString(FRDGResourceDumpContext::kBaseDir) / TEXT("Passes.json"), FILEWRITE_Append);
+	}
+
+	// Dump the pass' parameters
+	if (GDumpGPUPassParameters.GetValueOnRenderThread() != 0)
+	{
+		int32 PassParametersByteSize = 0;
+		{
+			const FShaderParametersMetadata* Metadata = Pass->GetParameters().GetMetadata();
+			if (Metadata)
+			{
+				GRDGResourceDumpContext.Dump(Metadata);
+				PassParametersByteSize = Metadata->GetSize();
+			}
+		}
+
+		if (PassParametersByteSize == 0 && Pass->GetParameters().GetLayoutPtr())
+		{
+			PassParametersByteSize = Pass->GetParameters().GetLayout().ConstantBufferSize;
+		}
 
-		FString JsonPath = GRDGResourceDumpContext.DumpingDirectoryPath + TEXT("Passes.json");
-		FFileHelper::SaveStringToFile(
-			OutputString, *JsonPath,
-			FFileHelper::EEncodingOptions::AutoDetect, &IFileManager::Get(), FILEWRITE_Append);
+		const uint8* PassParametersContent = Pass->GetParameters().GetContents();
+		if (PassParametersContent && !GRDGResourceDumpContext.IsDumped(PassParametersContent))
+		{
+			TArrayView<const uint8> ArrayView(PassParametersContent, PassParametersByteSize);
+			FString DumpFilePath = FRDGResourceDumpContext::kStructuresDir / FRDGResourceDumpContext::PtrToString(PassParametersContent) + TEXT(".bin");
+			GRDGResourceDumpContext.DumpBinaryToFile(ArrayView, DumpFilePath);
+			GRDGResourceDumpContext.SetDumped(PassParametersContent);
+		}
 	}
 
 	GRDGResourceDumpContext.PassesCount++;
@@ -622,4 +1540,192 @@ void FRDGBuilder::DumpResourcePassOutputs(const FRDGPass* Pass)
 	bInDebugPassScope = false;
 }
 
+#if RDG_DUMP_RESOURCES_AT_EACH_DRAW
+
+void FRDGBuilder::BeginPassDump(const FRDGPass* Pass)
+{
+	if (!GRDGResourceDumpContext.IsDumpingFrame())
+	{
+		return;
+	}
+
+	if (!GDumpGPUDraws.GetValueOnRenderThread())
+	{
+		return;
+	}
+
+	if (!EnumHasAnyFlags(Pass->GetFlags(), ERDGPassFlags::Raster))
+	{
+		return;
+	}
+
+	if (!IsInRenderingThread())
+	{
+		UE_LOG(LogRendererCore, Warning, TEXT("Couldn't start dumping draw's resources for pass %s because not in the rendering thread"), Pass->GetEventName().GetTCHAR());
+		return;
+	}
+
+	check(GRDGResourceDumpContext.DrawDumpingPass == nullptr);
+
+	if (GRDGResourceDumpContext.IsDumpingPass(Pass))
+	{
+		GRDGResourceDumpContext.DrawDumpingPass = Pass;
+		GRDGResourceDumpContext.DrawDumpCount = 0;
+	}
+}
+
+// static
+void FRDGBuilder::DumpDraw(const FRDGEventName& DrawEventName)
+{
+	if (!GRDGResourceDumpContext.IsDumpingFrame())
+	{
+		return;
+	}
+
+	if (!IsInRenderingThread())
+	{
+		UE_LOG(LogRendererCore, Warning, TEXT("Couldn't dump draw because not in the rendering thread"));
+		return;
+	}
+
+	if (!GRDGResourceDumpContext.DrawDumpingPass)
+	{
+		return;
+	}
+
+	const FRDGPass* Pass = GRDGResourceDumpContext.DrawDumpingPass;
+
+	FRHICommandListImmediate& RHICmdList = FRHICommandListExecutor::GetImmediateCommandList();
+
+	if (EnumHasAnyFlags(Pass->GetFlags(), ERDGPassFlags::Raster))
+	{
+		RHICmdList.EndRenderPass();
+	}
+
+	Pass->GetParameters().Enumerate([&](FRDGParameter Parameter)
+	{
+		switch (Parameter.GetType())
+		{
+		case UBMT_RENDER_TARGET_BINDING_SLOTS:
+		{
+			const FRenderTargetBindingSlots& RenderTargets = Parameter.GetAsRenderTargetBindingSlots();
+
+			RenderTargets.Enumerate([&](FRenderTargetBinding RenderTarget)
+			{
+				FRDGTextureRef Texture = RenderTarget.GetTexture();
+				FRDGTextureSRVDesc TextureSubResource = FRDGTextureSRVDesc::CreateForMipLevel(Texture, RenderTarget.GetMipIndex());
+				GRDGResourceDumpContext.DumpDrawTextureSubResource(
+					RHICmdList,
+					TextureSubResource,
+					ERHIAccess::RTV);
+			});
+
+			const FDepthStencilBinding& DepthStencil = RenderTargets.DepthStencil;
+
+			if (FRDGTextureRef Texture = DepthStencil.GetTexture())
+			{
+				FExclusiveDepthStencil DepthStencilAccess = DepthStencil.GetDepthStencilAccess();
+
+				if (DepthStencilAccess.IsDepthWrite())
+				{
+					FRDGTextureSRVDesc TextureSubResource = FRDGTextureSRVDesc::CreateForMipLevel(Texture, 0);
+					GRDGResourceDumpContext.DumpDrawTextureSubResource(
+						RHICmdList,
+						TextureSubResource,
+						ERHIAccess::RTV);
+				}
+
+				if (DepthStencilAccess.IsStencilWrite())
+				{
+					FRDGTextureSRVDesc TextureSubResource = FRDGTextureSRVDesc::CreateWithPixelFormat(Texture, PF_X24_G8);
+					GRDGResourceDumpContext.DumpDrawTextureSubResource(
+						RHICmdList,
+						TextureSubResource,
+						ERHIAccess::RTV);
+				}
+			}
+		}
+		break;
+		}
+	});
+
+	if (EnumHasAnyFlags(Pass->GetFlags(), ERDGPassFlags::Raster))
+	{
+		RHICmdList.BeginRenderPass(Pass->GetParameters().GetRenderPassInfo(), Pass->GetName());
+	}
+
+	// Dump the draw even name
+	{
+		TSharedPtr<FJsonObject> JsonObject = MakeShareable(new FJsonObject);
+		JsonObject->SetStringField(TEXT("DrawName"), DrawEventName.GetTCHAR());
+
+		FString DumpFilePath = FRDGResourceDumpContext::kPassesDir / FString::Printf(TEXT("Pass.%016x.Draws.json"), FRDGResourceDumpContext::PtrToUint(Pass));
+		GRDGResourceDumpContext.DumpJsonToFile(JsonObject, DumpFilePath, FILEWRITE_Append);
+	}
+
+	GRDGResourceDumpContext.DrawDumpCount++;
+
+	if (GRDGResourceDumpContext.DrawDumpCount % 10 == 0)
+	{
+		UE_LOG(LogRendererCore, Display, TEXT("Dumped %d draws' resources"), GRDGResourceDumpContext.DrawDumpCount);
+		return;
+	}
+}
+
+void FRDGBuilder::EndPassDump(const FRDGPass* Pass)
+{
+	if (!GRDGResourceDumpContext.IsDumpingFrame())
+	{
+		return;
+	}
+
+	if (!IsInRenderingThread())
+	{
+		return;
+	}
+
+	if (!GRDGResourceDumpContext.DrawDumpingPass)
+	{
+		return;
+	}
+
+	check(Pass == GRDGResourceDumpContext.DrawDumpingPass);
+
+	// Output how many draw has been dump for this pass.
+	if (GRDGResourceDumpContext.DrawDumpCount > 0)
+	{
+		TSharedPtr<FJsonObject> JsonObject = MakeShareable(new FJsonObject);
+		JsonObject->SetStringField(TEXT("EventName"), Pass->GetEventName().GetTCHAR());
+		JsonObject->SetStringField(TEXT("Pointer"), FString::Printf(TEXT("%016x"), FRDGResourceDumpContext::PtrToUint(Pass)));
+		JsonObject->SetNumberField(TEXT("DrawCount"), GRDGResourceDumpContext.DrawDumpCount);
+
+		GRDGResourceDumpContext.DumpJsonToFile(JsonObject, FString(FRDGResourceDumpContext::kBaseDir) / TEXT("PassDrawCounts.json"), FILEWRITE_Append);
+
+		UE_LOG(LogRendererCore, Display, TEXT("Completed dump of %d draws for pass: %s"), GRDGResourceDumpContext.DrawDumpCount, Pass->GetEventName().GetTCHAR());
+	}
+
+	GRDGResourceDumpContext.DrawDumpingPass = nullptr;
+	GRDGResourceDumpContext.DrawDumpCount = 0;
+}
+
+// static
+bool FRDGBuilder::IsDumpingDraws()
+{
+	if (!GRDGResourceDumpContext.IsDumpingFrame())
+	{
+		return false;
+	}
+
+	return GDumpGPUDraws.GetValueOnRenderThread() != 0;
+}
+
+#endif // RDG_DUMP_RESOURCES_AT_EACH_DRAW
+
+#else //! RDG_DUMP_RESOURCES
+
+bool IsDumpingRDGResources()
+{
+	return false;
+}
+
 #endif //! RDG_DUMP_RESOURCES
\ No newline at end of file
diff --git a/Engine/Source/Runtime/RenderCore/Private/RenderGraphPrivate.h b/Engine/Source/Runtime/RenderCore/Private/RenderGraphPrivate.h
--- a/Engine/Source/Runtime/RenderCore/Private/RenderGraphPrivate.h
+++ b/Engine/Source/Runtime/RenderCore/Private/RenderGraphPrivate.h
@@ -196,6 +196,8 @@ FORCEINLINE bool IsRenderPassMergeEnabled()
 	return GRDGMergeRenderPasses != 0 && !IsImmediateMode();
 }
 
+bool IsDumpingRDGResources();
+
 FORCEINLINE bool IsParallelExecuteEnabled()
 {
 	return GRDGParallelExecute > 0
@@ -203,7 +205,11 @@ FORCEINLINE bool IsParallelExecuteEnabled()
 		&& !IsImmediateMode()
 		&& !GRDGDebug
 		&& !GRDGTransitionLog
-		&& GRHISupportsMultithreadedShaderCreation;
+		&& GRHISupportsMultithreadedShaderCreation
+#if WITH_DUMPGPU
+		&& !IsDumpingRDGResources()
+#endif
+		;
 }
 
 template <typename ResourceRegistryType, typename FunctionType>
diff --git a/Engine/Source/Runtime/RenderCore/Public/RenderGraphBuilder.h b/Engine/Source/Runtime/RenderCore/Public/RenderGraphBuilder.h
--- a/Engine/Source/Runtime/RenderCore/Public/RenderGraphBuilder.h
+++ b/Engine/Source/Runtime/RenderCore/Public/RenderGraphBuilder.h
@@ -300,10 +300,20 @@ public:
 	FRDGBlackboard Blackboard;
 
 #if RDG_DUMP_RESOURCES
-	static void BeginResourceDump(const TArray<FString>& Args);
+	static FString BeginResourceDump(const TArray<FString>& Args);
 	static void EndResourceDump();
 #endif
 
+#if RDG_DUMP_RESOURCES_AT_EACH_DRAW
+	static void DumpDraw(const FRDGEventName& DrawEventName);
+	static bool IsDumpingDraws();
+#else
+	static inline bool IsDumpingDraws()
+	{
+		return false;
+	}
+#endif
+
 	//////////////////////////////////////////////////////////////////////////
 	// Deprecated Functions
 	UE_DEPRECATED(5.0, "PreallocateTexture has been renamed to ConvertToExternalTexture")
@@ -710,6 +720,11 @@ private:
 
 #if RDG_DUMP_RESOURCES
 	void DumpResourcePassOutputs(const FRDGPass* Pass);
+
+#if RDG_DUMP_RESOURCES_AT_EACH_DRAW
+	void BeginPassDump(const FRDGPass* Pass);
+	void EndPassDump(const FRDGPass* Pass);
+#endif
 #endif
 
 #if RDG_ENABLE_DEBUG
diff --git a/Engine/Source/Runtime/RenderCore/Public/RenderGraphDefinitions.h b/Engine/Source/Runtime/RenderCore/Public/RenderGraphDefinitions.h
--- a/Engine/Source/Runtime/RenderCore/Public/RenderGraphDefinitions.h
+++ b/Engine/Source/Runtime/RenderCore/Public/RenderGraphDefinitions.h
@@ -8,7 +8,7 @@
 /** DEFINES */
 
 /** Whether render graph debugging is enabled. */
-#define RDG_ENABLE_DEBUG (!UE_BUILD_SHIPPING && !UE_BUILD_TEST)
+#define RDG_ENABLE_DEBUG (!UE_BUILD_SHIPPING)
 
 /** Performs the operation if RDG_ENABLE_DEBUG is enabled. Useful for one-line checks without explicitly wrapping in #if. */
 #if RDG_ENABLE_DEBUG
@@ -29,6 +29,12 @@
 	#define IF_RDG_ENABLE_TRACE(Op)
 #endif
 
+/** Allows to dump all RDG resources of a frame. */
+#define RDG_DUMP_RESOURCES (WITH_DUMPGPU)
+
+/** Allows to dump RDG resource after each draw call. */
+#define RDG_DUMP_RESOURCES_AT_EACH_DRAW (RDG_DUMP_RESOURCES)
+
 /** The type of GPU events the render graph system supports.
  *  RDG_EVENTS == 0 means there is no string processing at all.
  *  RDG_EVENTS == 1 means the format component of the event name is stored as a const TCHAR*.
@@ -39,7 +45,7 @@
 #define RDG_EVENTS_STRING_COPY 2
 
 /** Whether render graph GPU events are enabled. */
-#if WITH_PROFILEGPU
+#if WITH_PROFILEGPU || RDG_DUMP_RESOURCES
 	#define RDG_EVENTS RDG_EVENTS_STRING_COPY
 #elif RHI_WANT_BREADCRUMB_EVENTS
 	#define RDG_EVENTS RDG_EVENTS_STRING_REF
@@ -71,9 +77,6 @@
 	#define IF_RDG_CMDLIST_STATS(Op)
 #endif
 
-/** Allows to dump all RDG resources of a frame. */
-#define RDG_DUMP_RESOURCES (!UE_BUILD_SHIPPING)
-
 /** ENUMS */
 
 enum class ERDGBuilderFlags
diff --git a/Engine/Source/Runtime/RenderCore/Public/RenderGraphParameter.h b/Engine/Source/Runtime/RenderCore/Public/RenderGraphParameter.h
--- a/Engine/Source/Runtime/RenderCore/Public/RenderGraphParameter.h
+++ b/Engine/Source/Runtime/RenderCore/Public/RenderGraphParameter.h
@@ -213,8 +213,10 @@ class RENDERCORE_API FRDGParameterStruct
 {
 public:
 	template <typename ParameterStructType>
-	explicit FRDGParameterStruct(const ParameterStructType* Parameters)
-		: FRDGParameterStruct(Parameters, ParameterStructType::FTypeInfo::GetStructMetadata()->GetLayoutPtr())
+	explicit FRDGParameterStruct(const ParameterStructType* Parameters, const FShaderParametersMetadata* InParameterMetadata)
+		: Contents(reinterpret_cast<const uint8*>(Parameters))
+		, Layout(InParameterMetadata->GetLayoutPtr())
+		, Metadata(InParameterMetadata)
 	{}
 
 	explicit FRDGParameterStruct(const void* InContents, const FRHIUniformBufferLayout* InLayout)
@@ -231,6 +233,8 @@ public:
 	const FRHIUniformBufferLayout& GetLayout() const { return *Layout; }
 	const FRHIUniformBufferLayout* GetLayoutPtr() const { return Layout; }
 
+	const FShaderParametersMetadata* GetMetadata() const { return Metadata; }
+
 	/** Helpful forwards from the layout. */
 	FORCEINLINE bool HasRenderTargets() const   { return Layout->HasRenderTargets(); }
 	FORCEINLINE bool HasExternalOutputs() const { return Layout->HasExternalOutputs(); }
@@ -288,6 +292,7 @@ private:
 
 	const uint8* Contents;
 	FUniformBufferLayoutRHIRef Layout;
+	const FShaderParametersMetadata* Metadata = nullptr;
 
 	friend class FRDGPass;
 };
@@ -321,12 +326,12 @@ public:
 template <typename TParameterStruct>
 FORCEINLINE static FRHIRenderPassInfo GetRenderPassInfo(TParameterStruct* Parameters)
 {
-	return FRDGParameterStruct(Parameters).GetRenderPassInfo();
+	return FRDGParameterStruct(Parameters, TParameterStruct::FTypeInfo::GetStructMetadata()).GetRenderPassInfo();
 }
 
 /** Helper function to get RHI global uniform buffers out of a pass parameters struct. */
 template <typename TParameterStruct>
 FORCEINLINE static FUniformBufferStaticBindings GetStaticUniformBuffers(TParameterStruct* Parameters)
 {
-	return FRDGParameterStruct(Parameters).GetStaticUniformBuffers();
+	return FRDGParameterStruct(Parameters, TParameterStruct::FTypeInfo::GetStructMetadata()).GetStaticUniformBuffers();
 }
\ No newline at end of file
diff --git a/Engine/Source/Runtime/RenderCore/Public/RenderGraphPass.h b/Engine/Source/Runtime/RenderCore/Public/RenderGraphPass.h
--- a/Engine/Source/Runtime/RenderCore/Public/RenderGraphPass.h
+++ b/Engine/Source/Runtime/RenderCore/Public/RenderGraphPass.h
@@ -539,7 +539,7 @@ public:
 		const ParameterStructType* InParameterStruct,
 		ERDGPassFlags InPassFlags,
 		ExecuteLambdaType&& InExecuteLambda)
-		: FRDGPass(MoveTemp(InName), FRDGParameterStruct(InParameterStruct, InParameterMetadata->GetLayoutPtr()), InPassFlags)
+		: FRDGPass(MoveTemp(InName), FRDGParameterStruct(InParameterStruct, InParameterMetadata), InPassFlags)
 		, ExecuteLambda(MoveTemp(InExecuteLambda))
 #if RDG_ENABLE_DEBUG
 		, DebugParameterStruct(InParameterStruct)
@@ -585,7 +585,7 @@ class FRDGSentinelPass final
 {
 public:
 	FRDGSentinelPass(FRDGEventName&& Name, ERDGPassFlags InPassFlagsToAdd = ERDGPassFlags::None)
-		: FRDGPass(MoveTemp(Name), FRDGParameterStruct(&EmptyShaderParameters), ERDGPassFlags::NeverCull | InPassFlagsToAdd)
+		: FRDGPass(MoveTemp(Name), FRDGParameterStruct(&EmptyShaderParameters, FEmptyShaderParameters::FTypeInfo::GetStructMetadata()), ERDGPassFlags::NeverCull | InPassFlagsToAdd)
 	{
 		bSentinel = 1;
 	}
diff --git a/Engine/Source/Runtime/RenderCore/Public/RenderGraphResources.h b/Engine/Source/Runtime/RenderCore/Public/RenderGraphResources.h
--- a/Engine/Source/Runtime/RenderCore/Public/RenderGraphResources.h
+++ b/Engine/Source/Runtime/RenderCore/Public/RenderGraphResources.h
@@ -178,7 +178,7 @@ protected:
 	template <typename TParameterStruct>
 	explicit FRDGUniformBuffer(const TParameterStruct* InParameters, const TCHAR* InName)
 		: FRDGResource(InName)
-		, ParameterStruct(InParameters)
+		, ParameterStruct(InParameters, TParameterStruct::FTypeInfo::GetStructMetadata())
 	{}
 
 private:
@@ -778,6 +778,14 @@ struct FRDGBufferDesc
 		return Desc;
 	}
 
+	template<typename ParameterStruct>
+	static FRDGBufferDesc CreateStructuredDesc(uint32 NumElements)
+	{
+		FRDGBufferDesc Desc = CreateStructuredDesc(sizeof(ParameterStruct), NumElements);
+		Desc.Metadata = ParameterStruct::FTypeInfo::GetStructMetadata();
+		return Desc;
+	}
+
 	static FRDGBufferDesc CreateBufferDesc(uint32 BytesPerElement, uint32 NumElements)
 	{
 		FRDGBufferDesc Desc;
@@ -788,6 +796,14 @@ struct FRDGBufferDesc
 		return Desc;
 	}
 
+	template<typename ParameterStruct>
+	static FRDGBufferDesc CreateBufferDesc(uint32 NumElements)
+	{
+		FRDGBufferDesc Desc = CreateBufferDesc(sizeof(ParameterStruct), NumElements);
+		Desc.Metadata = ParameterStruct::FTypeInfo::GetStructMetadata();
+		return Desc;
+	}
+
 	static FRDGBufferDesc CreateByteAddressDesc(uint32 NumBytes)
 	{
 		check(NumBytes % 4 == 0);
@@ -799,6 +815,14 @@ struct FRDGBufferDesc
 		return Desc;
 	}
 
+	template<typename ParameterStruct>
+	static FRDGBufferDesc CreateByteAddressDesc(uint32 NumElements)
+	{
+		FRDGBufferDesc Desc = CreateByteAddressDesc(sizeof(ParameterStruct) * NumElements);
+		Desc.Metadata = ParameterStruct::FTypeInfo::GetStructMetadata();
+		return Desc;
+	}
+
 	static FRDGBufferDesc CreateUploadDesc(uint32 BytesPerElement, uint32 NumElements)
 	{
 		FRDGBufferDesc Desc;
@@ -809,6 +833,14 @@ struct FRDGBufferDesc
 		return Desc;
 	}
 
+	template<typename ParameterStruct>
+	static FRDGBufferDesc CreateUploadDesc(uint32 NumElements)
+	{
+		FRDGBufferDesc Desc = CreateUploadDesc(sizeof(ParameterStruct), NumElements);
+		Desc.Metadata = ParameterStruct::FTypeInfo::GetStructMetadata();
+		return Desc;
+	}
+
 	/** Returns the total number of bytes allocated for a such buffer. */
 	uint32 GetTotalNumBytes() const
 	{
@@ -840,6 +872,9 @@ struct FRDGBufferDesc
 
 	/** The underlying RHI type to use. */
 	EUnderlyingType UnderlyingType = EUnderlyingType::VertexBuffer;
+
+	/** Meta data of the layout of the buffer for debugging purposes. */
+	const FShaderParametersMetadata* Metadata = nullptr;
 };
 
 inline const TCHAR* GetBufferUnderlyingTypeName(FRDGBufferDesc::EUnderlyingType BufferType)
