Add DumpGPU viewer script for Linux + Mac

Tested by:

 a) Running Engine/Binaries/Linux/UnrealEditor Samples/Games/ShooterGame/ShooterGame.uproject
 b) Playing game in editor
 c) Running DumpGPU command in console
 d) Folder pops up populated with DumpGPU html files and data plus OpenGPUDumpViewer.sh
 e) Double click on OpenGPUDumpViewer.sh and GPU Viewer opens in chrome and appears to work

Also tested script on Mac by copying entire folder to Mac and executing OpenGPUDumpViewer.sh

TODO:
  Better error messages when Chrome browser isn't installed

We appear to need the --allow-file-access-from-files chrome option and not sure how to get that functionality on other browsers via command line

diff --git a/Engine/Extras/GPUDumpViewer/OpenGPUDumpViewer.sh b/Engine/Extras/GPUDumpViewer/OpenGPUDumpViewer.sh
new file mode 100755
--- /dev/null
+++ b/Engine/Extras/GPUDumpViewer/OpenGPUDumpViewer.sh
@@ -0,0 +1,71 @@
+#!/usr/bin/env bash
+# Copyright Epic Games, Inc. All Rights Reserved.
+
+set -eu
+
+UNAMEOS="$(uname -s)"
+
+SCRIPT_DIR=$(cd "$(dirname "${BASH_SOURCE[0]}")" ; pwd)
+CHROME_USER_DATA=$(mktemp -d -t GPUDumpViewerUserData-XXXXXX)
+
+APPS=()
+APPS+=(google-chrome-stable)
+APPS+=(google-chrome)
+
+if [[ "${UNAMEOS}" =~ "Darwin" ]]; then
+	APPS+=(open)
+fi
+
+CMD=
+for val in "${APPS[@]}"; do
+	CMD="$(command -v "${val}")" || true
+	if [[ -n "${CMD}" ]]; then
+		break
+	fi
+done
+
+if [[ -z "${CMD}" ]]; then
+	echo "ERROR: Browser launch command not found"
+	exit 1
+fi
+
+ARGS=("${CMD}")
+
+if [[ "${CMD}" == *open ]]; then
+	ARGS+=(-a "google chrome")
+	ARGS+=(--new)
+	ARGS+=(-W)
+fi
+
+ARGS+=("file://${SCRIPT_DIR}/GPUDumpViewer.html")
+
+if [[ "${CMD}" == *open ]]; then
+	ARGS+=(--args)
+fi
+
+# --allow-file-access-from-files allow to load a file from a file:// webpage required for GPUDumpViewer.html to work.
+# --user-data-dir is required to force chrome to open a new instance so that --allow-file-access-from-files is honored.
+ARGS+=(--allow-file-access-from-files)
+ARGS+=(--new-window)
+ARGS+=(--incognito)
+ARGS+=("--user-data-dir=${CHROME_USER_DATA}")
+
+echo "Executing:"
+echo
+
+echo "${ARGS[0]} \\"
+for ((i=1; i < ${#ARGS[@]}; i++ )); do
+	echo "  ${ARGS[$i]} \\";
+done
+echo
+
+"${ARGS[@]}"
+
+echo
+echo "Closing ${CMD}..."
+
+if [[ -n "${CHROME_USER_DATA}" ]]; then
+	# Wait for 2s to shut down so that CHROME_USER_DATA can be deleted completely
+	sleep 2
+	rm -rf "${CHROME_USER_DATA}"
+fi
diff --git a/Engine/Source/Runtime/RenderCore/Private/RenderGraphDumpResources.cpp b/Engine/Source/Runtime/RenderCore/Private/RenderGraphDumpResources.cpp
--- a/Engine/Source/Runtime/RenderCore/Private/RenderGraphDumpResources.cpp
+++ b/Engine/Source/Runtime/RenderCore/Private/RenderGraphDumpResources.cpp
@@ -1296,12 +1296,16 @@ FString FRDGBuilder::BeginResourceDump(const TArray<FString>& Args)
 	// Copy the viewer
 	if (NewResourceDumpContext.bEnableDiskWrite)
 	{
-		const TCHAR* OpenGPUDumpViewerWindowsName = TEXT("OpenGPUDumpViewer.bat");
+#if PLATFORM_WINDOWS
+		const TCHAR* OpenGPUDumpViewerName = TEXT("OpenGPUDumpViewer.bat");
+#else
+		const TCHAR* OpenGPUDumpViewerName = TEXT("OpenGPUDumpViewer.sh");
+#endif
 		const TCHAR* ViewerHTML = TEXT("GPUDumpViewer.html");
 		FString DumpGPUViewerSourcePath = FPaths::EngineDir() + FString(TEXT("Extras")) / TEXT("GPUDumpViewer");
 
 		PlatformFile.CopyFile(*(NewResourceDumpContext.DumpingDirectoryPath / ViewerHTML), *(DumpGPUViewerSourcePath / ViewerHTML));
-		PlatformFile.CopyFile(*(NewResourceDumpContext.DumpingDirectoryPath / OpenGPUDumpViewerWindowsName), *(DumpGPUViewerSourcePath / OpenGPUDumpViewerWindowsName));
+		PlatformFile.CopyFile(*(NewResourceDumpContext.DumpingDirectoryPath / OpenGPUDumpViewerName), *(DumpGPUViewerSourcePath / OpenGPUDumpViewerName));
 	}
 
 	ENQUEUE_RENDER_COMMAND(FStartGPUDump)(
