Fixes OpenGPUDumpViewer.bat when a parent director have a space

diff --git a/Engine/Extras/GPUDumpViewer/OpenGPUDumpViewer.bat b/Engine/Extras/GPUDumpViewer/OpenGPUDumpViewer.bat
--- a/Engine/Extras/GPUDumpViewer/OpenGPUDumpViewer.bat
+++ b/Engine/Extras/GPUDumpViewer/OpenGPUDumpViewer.bat
@@ -8,7 +8,7 @@ echo "Opening chrome..."
 
 REM --allow-file-access-from-files allow to load a file from a file:// webpage required for GPUDumpViewer.html to work.
 REM --user-data-dir is required to force chrome to open a new instance so that --allow-file-access-from-files is honored.
-"C:\Program Files\Google\Chrome\Application\chrome.exe" file://%CWD%/GPUDumpViewer.html --allow-file-access-from-files --new-window --incognito --user-data-dir=%CHROME_USER_DATA%
+"C:\Program Files\Google\Chrome\Application\chrome.exe" "file://%CWD%/GPUDumpViewer.html" --allow-file-access-from-files --new-window --incognito --user-data-dir="%CHROME_USER_DATA%"
 
 echo "Closing chrome..."
 
