Fixes DumpGPU's resource search by savings all resource descriptors into a single Base/ResourceDesc.json file

diff --git a/Engine/Extras/GPUDumpViewer/GPUDumpViewer.html b/Engine/Extras/GPUDumpViewer/GPUDumpViewer.html
--- a/Engine/Extras/GPUDumpViewer/GPUDumpViewer.html
+++ b/Engine/Extras/GPUDumpViewer/GPUDumpViewer.html
@@ -18,6 +18,7 @@ const k_null_json_ptr = "0000000000000000";
 var g_infos = {};
 var g_dump_cvars = {};
 var g_passes = [];
+var g_descs = {};
 
 var g_view = null;
 
@@ -171,9 +172,9 @@ function load_json_dict_sequence(relative_path)
 	return dict_sequence;
 }
 
-function load_resource_desc(resource_name)
+function get_resource_desc(unique_resource_name)
 {
-	return load_json(`Resources/${resource_name}.json`);
+	return g_descs[unique_resource_name];
 }
 
 function load_structure_metadata(structure_ptr)
@@ -450,14 +451,18 @@ function display_pass_hierarchy()
 		if (search_resource != '')
 		{
 			show_resource = false;
-			pass_data['InputResources'].forEach(function(output_name) {
-				var subresource_unique_name = parse_subresource_unique_name(output_name)['resource'];
-				show_resource = show_resource || subresource_unique_name.toLowerCase().includes(search_resource.toLowerCase());
-			});
-			pass_data['OutputResources'].forEach(function(output_name) {
-				var subresource_unique_name = parse_subresource_unique_name(output_name)['resource'];
-				show_resource = show_resource || subresource_unique_name.toLowerCase().includes(search_resource.toLowerCase());
-			});
+			for (var subresource_unique_name of pass_data['InputResources'])
+			{
+				var resource_unique_name = parse_subresource_unique_name(subresource_unique_name)['resource'];
+				var resource_name = get_resource_desc(resource_unique_name)['Name'];
+				show_resource = show_resource || resource_name.toLowerCase().includes(search_resource.toLowerCase());
+			}
+			for (var subresource_unique_name of pass_data['OutputResources'])
+			{
+				var resource_unique_name = parse_subresource_unique_name(subresource_unique_name)['resource'];
+				var resource_name = get_resource_desc(resource_unique_name)['Name'];
+				show_resource = show_resource || resource_name.toLowerCase().includes(search_resource.toLowerCase());
+			}
 		}
 
 		var has_input_or_outputs = pass_data['InputResources'].length > 0 || pass_data['OutputResources'].length > 0;
@@ -612,7 +617,7 @@ class PassView extends IView
 		for (const subresource_unique_name of this.pass_data['InputResources'])
 		{
 			var subresource_info = parse_subresource_unique_name(subresource_unique_name);
-			var resource_desc = load_resource_desc(subresource_info['resource']);
+			var resource_desc = get_resource_desc(subresource_info['resource']);
 
 			var name = subresource_unique_name;
 			var href = null;
@@ -654,7 +659,7 @@ class PassView extends IView
 		for (const subresource_unique_name of this.pass_data['OutputResources'])
 		{
 			var subresource_info = parse_subresource_unique_name(subresource_unique_name);
-			var resource_desc = load_resource_desc(subresource_info['resource']);
+			var resource_desc = get_resource_desc(subresource_info['resource']);
 
 			var name = subresource_unique_name;
 			var href = null;
@@ -747,7 +752,7 @@ class PassView extends IView
 		document.getElementById('resource_pass_modifying_outter').style.display = 'block';
 
 		var subresource_unique_name = get_subresource_unique_name(this.resource_view.subresource_version_info);
-		var resource_desc = load_resource_desc(this.resource_view.subresource_version_info['resource']);
+		var resource_desc = get_resource_desc(this.resource_view.subresource_version_info['resource']);
 
 		var display_list = [];
 		for (var pass_id = 0; pass_id < g_passes.length; pass_id++)
@@ -788,7 +793,7 @@ class PassView extends IView
 		document.getElementById('resource_pass_reading_outter').style.display = 'block';
 
 		var subresource_unique_name = get_subresource_unique_name(this.resource_view.subresource_version_info);
-		var resource_desc = load_resource_desc(this.resource_view.subresource_version_info['resource']);
+		var resource_desc = get_resource_desc(this.resource_view.subresource_version_info['resource']);
 
 		var display_list = [];
 		for (var pass_id = 0; pass_id < g_passes.length; pass_id++)
@@ -876,7 +881,7 @@ function display_pass_internal(pass_id)
 
 function display_resource_internal(subresource_version_info)
 {
-	var resource_desc = load_resource_desc(subresource_version_info['resource']);
+	var resource_desc = get_resource_desc(subresource_version_info['resource']);
 	if (!resource_desc)
 	{
 		return;
@@ -2968,7 +2973,7 @@ function display_texture(subresource_version_info)
 		return;
 	}
 
-	var resource_desc = load_resource_desc(subresource_version_info['resource']);
+	var resource_desc = get_resource_desc(subresource_version_info['resource']);
 	g_view.set_resource_view(new TextureView(subresource_version_info, resource_desc));
 }
 
@@ -3803,21 +3808,32 @@ function init_page()
 				g_passes[pass_id]['EventName'] = 'Unnamed pass';
 			}
 		}
-	}
 
-	var pass_draw_counts = load_json_dict_sequence('Base/PassDrawCounts.json');
-	if (pass_draw_counts)
-	{
-		pass_draw_counts.forEach(function(pass_draw_count_data) {
-			for (var pass_id = 0; pass_id < g_passes.length; pass_id++)
-			{
-				if (g_passes[pass_id]['Pointer'] == pass_draw_count_data['Pointer'])
+		var pass_draw_counts = load_json_dict_sequence('Base/PassDrawCounts.json');
+		if (pass_draw_counts)
+		{
+			pass_draw_counts.forEach(function(pass_draw_count_data) {
+				for (var pass_id = 0; pass_id < g_passes.length; pass_id++)
 				{
-					g_passes[pass_id]['DrawCount'] = pass_draw_count_data['DrawCount'];
-					break;
+					if (g_passes[pass_id]['Pointer'] == pass_draw_count_data['Pointer'])
+					{
+						g_passes[pass_id]['DrawCount'] = pass_draw_count_data['DrawCount'];
+						break;
+					}
 				}
-			}
-		});
+			});
+		}
+	}
+
+	// Load all the resource descriptors
+	{
+		g_descs = {};
+
+		var desc_list = load_json_dict_sequence('Base/ResourceDescs.json');
+		for (const desc of desc_list)
+		{
+			g_descs[desc['UniqueResourceName']] = desc;
+		}
 	}
 
 	init_top();
@@ -3846,7 +3862,7 @@ function init_page()
 			for (var subresource_unique_name of pass_data['OutputResources'])
 			{
 				var subresource_info = parse_subresource_unique_name(subresource_unique_name);
-				var resource_desc = load_resource_desc(subresource_info['resource']);
+				var resource_desc = get_resource_desc(subresource_info['resource']);
 
 				if (resource_desc === null)
 				{
diff --git a/Engine/Source/Runtime/RenderCore/Private/RenderGraphDumpResources.cpp b/Engine/Source/Runtime/RenderCore/Private/RenderGraphDumpResources.cpp
--- a/Engine/Source/Runtime/RenderCore/Private/RenderGraphDumpResources.cpp
+++ b/Engine/Source/Runtime/RenderCore/Private/RenderGraphDumpResources.cpp
@@ -532,7 +532,7 @@ struct FRDGResourceDumpContext
 		return JsonObject;
 	}
 
-	TSharedPtr<FJsonObject> ToJson(const TCHAR* Name, const FRDGTextureDesc& Desc)
+	TSharedPtr<FJsonObject> ToJson(const FString& UniqueResourceName, const TCHAR* Name, const FRDGTextureDesc& Desc)
 	{
 		FString PixelFormat = ToJson(Desc.Format);
 
@@ -540,6 +540,7 @@ struct FRDGResourceDumpContext
 
 		TSharedPtr<FJsonObject> JsonObject = MakeShareable(new FJsonObject);
 		JsonObject->SetStringField(TEXT("Name"), Name);
+		JsonObject->SetStringField(TEXT("UniqueResourceName"), UniqueResourceName);
 		JsonObject->SetNumberField(TEXT("ByteSize"), ResourceByteSize);
 		JsonObject->SetStringField(TEXT("Desc"), TEXT("FRDGTextureDesc"));
 		JsonObject->SetStringField(TEXT("Type"), GetTextureDimensionString(Desc.Dimension));
@@ -567,10 +568,11 @@ struct FRDGResourceDumpContext
 		return JsonObject;
 	}
 
-	TSharedPtr<FJsonObject> ToJson(const TCHAR* Name, const FRDGBufferDesc& Desc, int32 ResourceByteSize)
+	TSharedPtr<FJsonObject> ToJson(const FString& UniqueResourceName, const TCHAR* Name, const FRDGBufferDesc& Desc, int32 ResourceByteSize)
 	{
 		TSharedPtr<FJsonObject> JsonObject = MakeShareable(new FJsonObject);
 		JsonObject->SetStringField(TEXT("Name"), Name);
+		JsonObject->SetStringField(TEXT("UniqueResourceName"), UniqueResourceName);
 		JsonObject->SetNumberField(TEXT("ByteSize"), ResourceByteSize);
 		JsonObject->SetStringField(TEXT("Desc"), TEXT("FRDGBufferDesc"));
 		JsonObject->SetStringField(TEXT("Type"), GetBufferUnderlyingTypeName(Desc.UnderlyingType));
@@ -946,10 +948,8 @@ struct FRDGResourceDumpContext
 		// Dump the information of the texture to json file.
 		if (bDumpResourceInfos)
 		{
-			TSharedPtr<FJsonObject> JsonObject = ToJson(SubresourceDesc.Texture->Name, Desc);
-
-			FString ResourceInfoJsonPath = kResourcesDir / UniqueResourceName + TEXT(".json");
-			DumpJsonToFile(JsonObject, ResourceInfoJsonPath);
+			TSharedPtr<FJsonObject> JsonObject = ToJson(UniqueResourceName, SubresourceDesc.Texture->Name, Desc);
+			DumpJsonToFile(JsonObject, FString(FRDGResourceDumpContext::kBaseDir) / TEXT("ResourceDescs.json"), FILEWRITE_Append);
 		}
 
 		if (!SubresourceDumpDesc.IsDumpSupported())
@@ -1093,10 +1093,8 @@ struct FRDGResourceDumpContext
 		// Dump the information of the buffer to json file.
 		if (bDumpResourceInfos)
 		{
-			TSharedPtr<FJsonObject> JsonObject = ToJson(Buffer->Name, Desc, ByteSize);
-
-			FString ResourceInfoJsonPath = kResourcesDir / UniqueResourceName + TEXT(".json");
-			DumpJsonToFile(JsonObject, ResourceInfoJsonPath);
+			TSharedPtr<FJsonObject> JsonObject = ToJson(UniqueResourceName, Buffer->Name, Desc, ByteSize);
+			DumpJsonToFile(JsonObject, FString(FRDGResourceDumpContext::kBaseDir) / TEXT("ResourceDescs.json"), FILEWRITE_Append);
 
 			if (Desc.Metadata && DumpBufferMode == 2)
 			{
@@ -1242,6 +1240,8 @@ FString FRDGBuilder::BeginResourceDump(const TArray<FString>& Args)
 		PlatformFile.CreateDirectoryTree(*(NewResourceDumpContext.DumpingDirectoryPath / FRDGResourceDumpContext::kBaseDir));
 		PlatformFile.CreateDirectoryTree(*(NewResourceDumpContext.DumpingDirectoryPath / FRDGResourceDumpContext::kResourcesDir));
 
+		NewResourceDumpContext.DumpStringToFile(TEXT(""), FString(FRDGResourceDumpContext::kBaseDir) / TEXT("Passes.json"));
+		NewResourceDumpContext.DumpStringToFile(TEXT(""), FString(FRDGResourceDumpContext::kBaseDir) / TEXT("ResourceDescs.json"));
 		NewResourceDumpContext.DumpStringToFile(TEXT(""), FString(FRDGResourceDumpContext::kBaseDir) / TEXT("PassDrawCounts.json"));
 	}
 
