Fixes DumpGPU when only 32bit chrome is installed

diff --git a/Engine/Extras/GPUDumpViewer/OpenGPUDumpViewer.bat b/Engine/Extras/GPUDumpViewer/OpenGPUDumpViewer.bat
--- a/Engine/Extras/GPUDumpViewer/OpenGPUDumpViewer.bat
+++ b/Engine/Extras/GPUDumpViewer/OpenGPUDumpViewer.bat
@@ -6,9 +6,15 @@ set CHROME_USER_DATA=%CWD%/.tmp_chrome_data/
 
 echo "Opening chrome..."
 
+
+set CHROME="C:\Program Files (x86)\Google\Chrome\Application\chrome.exe"
+if exist "C:\Program Files\Google\Chrome\Application\chrome.exe" (
+	set CHROME="C:\Program Files\Google\Chrome\Application\chrome.exe"
+)
+
 REM --allow-file-access-from-files allow to load a file from a file:// webpage required for GPUDumpViewer.html to work.
 REM --user-data-dir is required to force chrome to open a new instance so that --allow-file-access-from-files is honored.
-"C:\Program Files\Google\Chrome\Application\chrome.exe" "file://%CWD%/GPUDumpViewer.html" --allow-file-access-from-files --new-window --incognito --user-data-dir="%CHROME_USER_DATA%"
+%CHROME% "file://%CWD%/GPUDumpViewer.html" --allow-file-access-from-files --new-window --incognito --user-data-dir="%CHROME_USER_DATA%"
 
 echo "Closing chrome..."
 
