Fixes texture visualization in the GPUDumpViewer for PF_R32G32_UINT

diff --git a/Engine/Extras/GPUDumpViewer/GPUDumpViewer.html b/Engine/Extras/GPUDumpViewer/GPUDumpViewer.html
--- a/Engine/Extras/GPUDumpViewer/GPUDumpViewer.html
+++ b/Engine/Extras/GPUDumpViewer/GPUDumpViewer.html
@@ -1511,7 +1511,7 @@ class TextureView extends ResourceView
 		var pixel_data_offset = texel_x + (extent['y'] - 1 - texel_y) * extent['x'];
 
 		var texel_value = [0];
-		for (var c = 1; c < this.subresource_desc['webgl_channel_count']; c++)
+		for (var c = 1; c < this.subresource_desc['raw_channel_count']; c++)
 		{
 			texel_value.push(0);
 		}
@@ -1546,22 +1546,22 @@ class TextureView extends ResourceView
 			var data = null;
 			if (this.subresource_desc['raw_channel_bit_depth'] === 8)
 			{
-				data = new Uint8Array(this.raw_texture_data, this.subresource_desc['webgl_channel_count'] * pixel_data_offset * 1, this.subresource_desc['webgl_channel_count']);
+				data = new Uint8Array(this.raw_texture_data, this.subresource_desc['raw_channel_count'] * pixel_data_offset * 1, this.subresource_desc['raw_channel_count']);
 			}
 			else if (this.subresource_desc['raw_channel_bit_depth'] === 16)
 			{
-				data = new Uint16Array(this.raw_texture_data, this.subresource_desc['webgl_channel_count'] * pixel_data_offset * 2, this.subresource_desc['webgl_channel_count']);
+				data = new Uint16Array(this.raw_texture_data, this.subresource_desc['raw_channel_count'] * pixel_data_offset * 2, this.subresource_desc['raw_channel_count']);
 			}
 			else if (this.subresource_desc['raw_channel_bit_depth'] === 32)
 			{
-				data = new Uint32Array(this.raw_texture_data, this.subresource_desc['webgl_channel_count'] * pixel_data_offset * 4, this.subresource_desc['webgl_channel_count']);
+				data = new Uint32Array(this.raw_texture_data, this.subresource_desc['raw_channel_count'] * pixel_data_offset * 4, this.subresource_desc['raw_channel_count']);
 			}
 			else
 			{
 				return;
 			}
 
-			for (var c = 0; c < this.subresource_desc['webgl_channel_count']; c++)
+			for (var c = 0; c < this.subresource_desc['raw_channel_count']; c++)
 			{
 				texel_value[c] = data[this.subresource_desc['raw_channel_swizzling'][c]];
 			}
@@ -1650,7 +1650,7 @@ class TextureView extends ResourceView
 			const channel_color = ['rgb(255, 38, 38)', 'rgb(38, 255, 38)', 'rgb(38, 187, 255)', 'white'];
 
 			var texel_string = [];
-			for (var c = 0; c < texel_value.length; c++)
+			for (var c = 0; c < this.subresource_desc['raw_channel_count']; c++)
 			{
 				texel_string.push(`${texel_value[c].toString()} (0x${texel_raw[c].toString(16)})`);
 			}
@@ -1661,7 +1661,7 @@ class TextureView extends ResourceView
 			}
 
 			var html = ``;
-			for (var c = 0; c < texel_value.length; c++)
+			for (var c = 0; c < this.subresource_desc['raw_channel_count']; c++)
 			{
 				html += `<span style="margin-left: 20px; color: ${channel_color[c]};">${channel_name[c]}: <span style="min-width: 170px; display: inline-block; text-align: right;">${texel_string[c]}</span></span>`;
 			}
@@ -2494,25 +2494,17 @@ void main(void)
 			shader_display_operation = `${shader_float_vector[raw_channel_count]}(${shader_display_operation}) / ${divide}.0`;
 		}
 
-		var shader_display_assignment = 'display.rgba';
-		var shader_fetch_return_operation = 'raw';
 		var webgl_channel_count = 0;
 		if (webgl_texture_src_format === gl.RED || webgl_texture_src_format === gl.RED_INTEGER)
 		{
-			shader_display_assignment = `display.r`;
-			shader_fetch_return_operation = `${shader_fetch_return_operation}.r`;
 			webgl_channel_count = 1;
 		}
 		else if (webgl_texture_src_format === gl.RG)
 		{
-			shader_display_assignment = `display.rg`;
-			shader_fetch_return_operation = `${shader_fetch_return_operation}.rg`;
 			webgl_channel_count = 2;
 		}
 		else if (webgl_texture_src_format === gl.RGB)
 		{
-			shader_display_assignment = `display.rgb`;
-			shader_fetch_return_operation = `${shader_fetch_return_operation}.rgb`;
 			webgl_channel_count = 3;
 		}
 		else if (webgl_texture_src_format === gl.RGBA || webgl_texture_src_format === gl.RGBA_INTEGER)
@@ -2543,12 +2535,22 @@ void main(void)
 			shader_fetch_operation = `texelFetch(texture0, texel_coord , 0)`;
 		}
 
+		const k_get_n_channels = new Array(
+			'',
+			'.r',
+			'.rg',
+			'.rgb',
+			'.rgba'
+		);
+
+		var shader_display_assignment = `display${k_get_n_channels[raw_channel_count]}`;
+
 		webgl_pixel_shader_private += `
 ${shader_sampler_return_type[raw_channel_count]} fetchTexel(vec2 uv)
 {
 	ivec2 texel_coord = ivec2(uv * vec2(textureSize(texture0, 0)));
 	${shader_sampler_return_type[4]} raw = ${shader_fetch_operation};
-	return ${shader_fetch_return_operation};
+	return raw${k_get_n_channels[raw_channel_count]};
 }
 
 in vec2 uv;
@@ -2739,15 +2741,15 @@ ${shader_display_assignment} = ${shader_display_operation};`;
 					{
 						for (var c = 0; c < this.subresource_desc['webgl_channel_count']; c++)
 						{
-							var i = this.subresource_desc['webgl_channel_count'] * (x + extent['x'] * y);
+							var texel_offset = x + extent['x'] * y;
 
 							var value = 0;
 							if (c < this.subresource_desc['raw_channel_count'])
 							{
-								var src = original_converted_data[i + this.subresource_desc['raw_channel_swizzling'][c]];
+								var src = original_converted_data[texel_offset * this.subresource_desc['raw_channel_count'] + this.subresource_desc['raw_channel_swizzling'][c]];
 								value = (src >> (8 * tex_id)) % 256;
 							}
-							converted_data[i + c] = value;
+							converted_data[texel_offset * this.subresource_desc['webgl_channel_count'] + c] = value;
 						}
 					}
 				}
@@ -2927,7 +2929,7 @@ void main(void) {
 		else
 		{
 			var user_frag_code = document.getElementById('texture_visualization_code_input').value;
-			shader_program = this.compile_screen_shader_program(this.subresource_desc['webgl_pixel_shader_private'] + 'void main(void) { display = vec4(0.0, 0.0, 0.0, 0.0);\n' + user_frag_code + `}`);
+			shader_program = this.compile_screen_shader_program(this.subresource_desc['webgl_pixel_shader_private'] + 'void main(void) { display = vec4(0.0, 0.0, 0.0, 0.0);\n' + user_frag_code + `\n}`);
 
 			if (!shader_program)
 			{
