Adds a viewer console that can detect errors in the dump such as incomplete dump.

diff --git a/Engine/Extras/GPUDumpViewer/GPUDumpViewer.html b/Engine/Extras/GPUDumpViewer/GPUDumpViewer.html
--- a/Engine/Extras/GPUDumpViewer/GPUDumpViewer.html
+++ b/Engine/Extras/GPUDumpViewer/GPUDumpViewer.html
@@ -100,10 +100,11 @@ function load_text_file(relative_path)
 		{
 			return request.responseText;
 		}
+		add_console_event('error', `couldn't load ${relative_path}: Returned HTTP status ${request.status}`);
 	}
 	catch (error)
 	{
-		return null;
+		add_console_event('error', `couldn't load ${relative_path}: ${error}`);
 	}
 	return null;
 }
@@ -125,11 +126,13 @@ function load_binary_file(relative_path, callback)
 			}
 			else
 			{
+				add_console_event('error', `couldn't load ${relative_path}: Returned HTTP status ${request.status}`);
 				callback(null);
 			}
 		};
 
 		request.onerror = function() {
+			add_console_event('error', `couldn't load ${relative_path}`);
 			callback(null);
 		};
 
@@ -137,6 +140,7 @@ function load_binary_file(relative_path, callback)
 	}
 	catch (error)
 	{
+		add_console_event('error', `couldn't load ${relative_path}: ${error}`);
 		callback(null);
 	}
 }
@@ -996,6 +1000,103 @@ function display_draw_output_resource(pass_id, subresource_unique_name, draw_id)
 }
 
 
+// -------------------------------------------------------------------- DISPLAY VIEWER CONSOLE
+
+var g_console_events = [];
+var g_console_error_count = 0;
+
+class ConsoleEvent
+{
+	constructor(type, message)
+	{
+		this.type = type;
+		this.message = message;
+	}
+}
+
+class ConsoleView extends IView
+{
+	constructor()
+	{
+		super();
+	}
+
+	setup_html(parent_dom)
+	{
+		parent_dom.innerHTML = `
+			<div class="main_div">
+				<div class="pass_title">Viewer Console</div>
+				<div id="console_events_pannel"></div>
+			</div>`;
+
+		document.title = 'Viewer Console';
+		this.update_console_events();
+	}
+
+	update_console_events()
+	{
+		var html = `
+			<table width="100%" class="pretty_table">`;
+
+		for (var console_event of g_console_events)
+		{
+			html += `
+				<tr class="${console_event.type}">
+					<td>${console_event.type}: ${console_event.message}</td>
+				</tr>`;
+		};
+
+		html += `
+			</table>`;
+
+		document.getElementById('console_events_pannel').innerHTML = html;
+	}
+
+	get navigations()
+	{
+		return [`display_console();`];
+	}
+}
+
+function update_console_button()
+{
+	if (document.getElementById('console_button') && g_console_error_count > 0)
+	{
+		document.getElementById('console_button').classList.add('error');
+		document.getElementById('console_button').innerHTML = `Console (${g_console_error_count} Errors)`;
+	}
+}
+
+function add_console_event(type, message)
+{
+	console.assert(['error', 'log'].includes(type));
+	g_console_events.push(new ConsoleEvent(type, message));
+
+	if (type == 'error')
+	{
+		g_console_error_count += 1;
+		update_console_button();
+	}
+
+	if (g_view instanceof ConsoleView)
+	{
+		g_view.update_console_events();
+	}
+}
+
+function display_console(tip_id)
+{
+	set_main_view(new ConsoleView());
+}
+
+function init_console()
+{
+	window.addEventListener('error', function(event) {
+		add_console_event('error', `${event.filename}:${event.fileno}: ${event.message}`);
+	});
+}
+
+
 // -------------------------------------------------------------------- DISPLAY INFOS
 
 class InfosView extends IView
@@ -1249,12 +1350,24 @@ function display_cvars()
 
 // -------------------------------------------------------------------- DISPLAY LOG
 
+function get_log_path()
+{
+	return `Base/${g_infos['Project']}.log`;
+}
+
 class LogView extends IView
 {
 	constructor()
 	{
 		super();
-		this.log = load_text_file(`Base/${g_infos['Project']}.log`).split('\n');
+
+		this.log = [];
+
+		var log = load_text_file(get_log_path());
+		if (log)
+		{
+			this.log = log.split('\n');
+		}
 	}
 
 	setup_html(parent_dom)
@@ -3790,6 +3903,9 @@ function init_top_buttons()
 		<a href="#display_infos();">Infos</a>
 		<a href="#display_cvars();">CVars</a>
 		<a href="#display_log();">Log</a>`;
+	document.getElementById('top_viewer_buttons').innerHTML = `
+		<a id="console_button" href="#display_console();">Console</a>`;
+	update_console_button();
 }
 
 function init_page()
@@ -3863,6 +3979,15 @@ function init_page()
 		}
 	}
 
+	if (!does_file_exists(get_log_path()))
+	{
+		add_console_event('error', `The dump does not contain ${get_log_path()}. The log is normally copied into the dump directory once the dump is completed. Failing to have is may be due to a premature end of the dumping process.`);
+	}
+	else
+	{
+		add_console_event('log', `Viewer init ok`);
+	}
+
 	init_top();
 	init_top_buttons();
 	display_pass_hierarchy();
@@ -4341,6 +4466,10 @@ input::placeholder
 	text-align: center;
 }
 
+.pretty_table tr.error td
+{
+	color: var(--error-txt-color);
+}
 
 .resource_desc tr td:not(.empty)
 {
@@ -4382,6 +4511,11 @@ input::placeholder
 .button_list input:not(.match_value):not(:active):hover
 {
 	background-color: var(--button-bg-color-hover);
+}
+
+.button_list a:not(.match_location_hash):not(.error):not(:active):hover,
+.button_list input:not(.match_value):not(:active):hover
+{
 	color: var(--button-txt-color-hover);
 }
 
@@ -4394,6 +4528,12 @@ input::placeholder
 	color: var(--button-txt-color-selected);
 }
 
+.button_list a.error,
+.button_list a.error:active
+{
+	color: var(--error-txt-color);
+}
+
 
 /* ----------------------------------------------------- top bar */
 
@@ -4405,7 +4545,8 @@ input::placeholder
 	display: inline-block;
 }
 
-#top_buttons
+#top_buttons,
+#top_viewer_buttons
 {
 	padding: 12px;
 	display: inline-block;
@@ -4605,12 +4746,13 @@ input::placeholder
 </style>
 
 
-<body onload="onresize_body(); init_page();" onresize="onresize_body();" onhashchange="navigate_to_hash();">
+<body onload="init_console(); onresize_body(); init_page();" onresize="onresize_body();" onhashchange="navigate_to_hash();">
 	<table cellpadding="0" cellspacing="0" border="0" id="onresize_body_table">
 		<tr>
 			<td colspan="2" id="onresize_body_top_bar">
 				<div id="top_bar"></div>
 				<div id="top_buttons" class="button_list"></div>
+				<div id="top_viewer_buttons" class="button_list"></div>
 			</td>
 		</tr>
 		<tr>
