Fixes an error in the GPU dump viewers when no passes has been dumped

diff --git a/Engine/Extras/GPUDumpViewer/GPUDumpViewer.html b/Engine/Extras/GPUDumpViewer/GPUDumpViewer.html
--- a/Engine/Extras/GPUDumpViewer/GPUDumpViewer.html
+++ b/Engine/Extras/GPUDumpViewer/GPUDumpViewer.html
@@ -186,6 +186,10 @@ function load_json(relative_path)
 function load_json_dict_sequence(relative_path)
 {
 	var text_file = load_text_file(relative_path);
+	if (text_file == '')
+	{
+		return new Array();
+	}
 	var dicts = text_file.split("}{");
 	var dict_sequence = [];
 	dicts.forEach(function(value, index, array) {
