Fixes a bug in the GPU viewer's texture viewer when the size of a texture is not multiple of 4

diff --git a/Engine/Extras/GPUDumpViewer/GPUDumpViewer.html b/Engine/Extras/GPUDumpViewer/GPUDumpViewer.html
--- a/Engine/Extras/GPUDumpViewer/GPUDumpViewer.html
+++ b/Engine/Extras/GPUDumpViewer/GPUDumpViewer.html
@@ -572,6 +572,7 @@ class ResourceView extends IView
 		super();
 		this.subresource_version_info = subresource_version_info;
 		this.resource_desc = resource_desc;
+		this.onload = function() { };
 	}
 
 	get navigations()
@@ -944,7 +945,22 @@ function display_resource_internal(subresource_version_info)
 
 	if (resource_desc['Type'] == 'Texture2D' || resource_desc['Type'] == 'Texture2DArray')
 	{
-		g_view.set_resource_view(new TextureView(subresource_version_info, resource_desc));
+		var new_texture_view = new TextureView(subresource_version_info, resource_desc);
+
+		// Keep the same zoom settings if the subresources are exactly the same extent
+		if (g_view && 'resource_view' in g_view && g_view.resource_view instanceof TextureView)
+		{
+			var zoom_settings = g_view.resource_view.get_zoom_settings();
+
+			if (new_texture_view.is_zoom_settings_compatible(zoom_settings))
+			{
+				new_texture_view.onload = function() {
+					new_texture_view.apply_zoom_settings(zoom_settings);
+				};
+			}
+		}
+
+		g_view.set_resource_view(new_texture_view);
 	}
 	else if (resource_desc['Type'] == 'VertexBuffer' || resource_desc['Type'] == 'StructuredBuffer')
 	{
@@ -1642,6 +1658,7 @@ class TextureView extends ResourceView
 				texture_view.resize_canvas(texture_view.pixel_scaling);
 				texture_view.process_texture_data_for_visualization();
 				texture_view.refresh_texture_visualization();
+				texture_view.onload();
 			}
 			else
 			{
@@ -1931,23 +1948,24 @@ class TextureView extends ResourceView
 		var [texel_x, texel_y] = this.get_texel_coordinate_of_mouse(event);
 
 		this.resize_canvas(new_pixel_scaling, texel_x, texel_y);
-
-		// Update the state of the zooming buttons
-		{
-			const pixel_scaling_names = {
-				0: 'Fit',
-				1: '1:1',
-				2: '2:1',
-				4: '4:1',
-				8: '8:1',
-				16: '16:1',
-			};
-			update_value_selection(document.getElementById('texture_visualization_change_pixel_scaling'), pixel_scaling_names[this.pixel_scaling]);
-		}
+		this.update_change_pixel_scaling_button();
 
 		return false;
 	}
 
+	update_change_pixel_scaling_button()
+	{
+		const pixel_scaling_names = {
+			0: 'Fit',
+			1: '1:1',
+			2: '2:1',
+			4: '4:1',
+			8: '8:1',
+			16: '16:1',
+		};
+		update_value_selection(document.getElementById('texture_visualization_change_pixel_scaling'), pixel_scaling_names[this.pixel_scaling]);
+	}
+
 	canvas_click(event)
 	{
 		if (!this.is_ready)
@@ -2131,6 +2149,44 @@ class TextureView extends ResourceView
 		this.pixel_scaling = new_pixel_scaling;
 	}
 
+	get_zoom_settings()
+	{
+		var subresource_extent = this.get_subresource_extent();
+
+		var canvas_viewport = document.getElementById('canvas_viewport');
+
+		var zoom_settings = {
+			'subresource_extent_x': subresource_extent['x'],
+			'subresource_extent_y': subresource_extent['y'],
+			'pixel_scaling': this.pixel_scaling,
+			'canvas_width': this.canvas.style.width,
+			'canvas_height': this.canvas.style.height,
+			'canvas_viewport_scoll_left': canvas_viewport.scrollLeft,
+			'canvas_viewport_scoll_top': canvas_viewport.scrollTop,
+		};
+		return zoom_settings;
+	}
+
+	is_zoom_settings_compatible(zoom_settings)
+	{
+		var subresource_extent = this.get_subresource_extent();
+
+		return (zoom_settings['subresource_extent_x'] == subresource_extent['x'] && zoom_settings['subresource_extent_y'] == subresource_extent['y']);
+	}
+
+	apply_zoom_settings(zoom_settings)
+	{
+		var canvas_viewport = document.getElementById('canvas_viewport');
+
+		this.pixel_scaling = zoom_settings['pixel_scaling'];
+		this.canvas.style.width = zoom_settings['canvas_width'];
+		this.canvas.style.height = zoom_settings['canvas_height'];
+		canvas_viewport.scrollLeft = zoom_settings['canvas_viewport_scoll_left'];
+		canvas_viewport.scrollTop = zoom_settings['canvas_viewport_scoll_top'];
+
+		this.update_change_pixel_scaling_button();
+	}
+
 	init_gl()
 	{
 		var gl_ctx_attributes = {
@@ -2170,6 +2226,9 @@ class TextureView extends ResourceView
 			});
 		}
 
+		// Set unpacking alignement to 1 to allow uploading texture size not multple of 4 
+		gl.pixelStorei(gl.UNPACK_ALIGNMENT, 1);
+
 		// Create vertex buffer
 		{
 			var vertices = [
@@ -3278,6 +3337,7 @@ class GenericBufferView extends ResourceView
 			{
 				buffer_view.raw_buffer_data = raw_buffer_data;
 				buffer_view.refresh_buffer_visualization(/* refresh header = */ true);
+				buffer_view.onload();
 			}
 			else
 			{
