Fixes GPU dump viewer not displaying buffers

diff --git a/Engine/Extras/GPUDumpViewer/GPUDumpViewer.html b/Engine/Extras/GPUDumpViewer/GPUDumpViewer.html
--- a/Engine/Extras/GPUDumpViewer/GPUDumpViewer.html
+++ b/Engine/Extras/GPUDumpViewer/GPUDumpViewer.html
@@ -943,7 +943,11 @@ function display_resource_internal(subresource_version_info)
 		return;
 	}
 
-	if (resource_desc['Type'] == 'Texture2D' || resource_desc['Type'] == 'Texture2DArray')
+	if (resource_desc['Desc'] == 'FRDGBufferDesc')
+	{
+		g_view.set_resource_view(new BufferView(subresource_version_info, resource_desc));
+	}
+	else if (resource_desc['Desc'] == 'FRDGTextureDesc' && (resource_desc['Type'] == 'Texture2D' || resource_desc['Type'] == 'Texture2DArray'))
 	{
 		var new_texture_view = new TextureView(subresource_version_info, resource_desc);
 
@@ -962,10 +966,6 @@ function display_resource_internal(subresource_version_info)
 
 		g_view.set_resource_view(new_texture_view);
 	}
-	else if (resource_desc['Type'] == 'VertexBuffer' || resource_desc['Type'] == 'StructuredBuffer')
-	{
-		g_view.set_resource_view(new BufferView(subresource_version_info, resource_desc));
-	}
 }
 
 
@@ -3830,7 +3830,6 @@ function display_pass_parameters(pass_id)
 	resource_desc['Name'] = 'PassParameters';
 	resource_desc['ByteSize'] = structure_metadata['Size'];
 	resource_desc['Desc'] = 'FRDGParameterStruct';
-	resource_desc['Type'] = 'Parameters';
 	resource_desc['BytesPerElement'] = resource_desc['ByteSize'];
 	resource_desc['NumElements'] = 1;
 	resource_desc['Metadata'] = pass_data['ParametersMetadata'];
