Fixes a bug where the texture visualization shader was being shared between same texture name of different pixel format

diff --git a/Engine/Extras/GPUDumpViewer/GPUDumpViewer.html b/Engine/Extras/GPUDumpViewer/GPUDumpViewer.html
--- a/Engine/Extras/GPUDumpViewer/GPUDumpViewer.html
+++ b/Engine/Extras/GPUDumpViewer/GPUDumpViewer.html
@@ -2394,6 +2394,8 @@ void main(void)
 			name = this.subresource_version_info['resource'];
 		}
 
+		name = `${name}.${this.resource_desc['Format']}`;
+
 		if (this.subresource_version_info['subresource'] && !this.subresource_version_info['subresource'].startsWith('mip'))
 		{
 			name = `${name}.${this.subresource_version_info['subresource']}`;
