Fix DumpGPU on Linux so that OpenGPUDumpViewer.sh opens the dump in the user's default browser

diff --git a/Engine/Extras/GPUDumpViewer/OpenGPUDumpViewer.sh b/Engine/Extras/GPUDumpViewer/OpenGPUDumpViewer.sh
--- a/Engine/Extras/GPUDumpViewer/OpenGPUDumpViewer.sh
+++ b/Engine/Extras/GPUDumpViewer/OpenGPUDumpViewer.sh
@@ -4,13 +4,59 @@
 set -eu
 
 UNAMEOS="$(uname -s)"
+ADDRESS="127.0.0.1"
+PORT="8000"
+URL="http://$ADDRESS:$PORT/GPUDumpViewer.html"
 
-SCRIPT_DIR=$(cd "$(dirname "${BASH_SOURCE[0]}")" ; pwd)
-CHROME_USER_DATA=$(mktemp -d -t GPUDumpViewerUserData-XXXXXX)
+SCRIPT=$(readlink -f "$0")
+# Absolute path this script is in, thus [REDACTED_USER_HOME]/bin
+SCRIPTPATH=$(dirname "$SCRIPT")
+
+pushd $SCRIPTPATH
+
+
+GetAllChildProcesses() {
+	local Children=$(ps -o pid= ppid "$1")
+
+	for PID in $Children
+	do
+		GetAllChildProcesses "$PID"
+	done
+
+	echo "$Children"
+}
+
+# Gather all the descendant children of this process, and first kill -TERM. If any child process
+# is still alive finally send a -KILL
+TermHandler() {
+	MaxWait=30
+	CurrentWait=0
+
+	ProcessesToKill=$(GetAllChildProcesses $$)
+	kill -s TERM $ProcessesToKill 2> /dev/null
+
+	ProcessesStillAlive=$(ps -o pid= -p $ProcessesToKill)
+
+	# Wait until all the processes have been gracefully killed, or max Wait time
+	while [ -n "$ProcessesStillAlive" ] && [ "$CurrentWait" -lt "$MaxWait" ]
+	do
+		CurrentWait=$((CurrentWait + 1))
+		sleep 1
+
+		ProcessesStillAlive=$(ps -o pid= -p $ProcessesToKill)
+	done
+
+	# If some processes are still alive after MaxWait, lets just force kill them
+	if [ -n "$ProcessesStillAlive" ]; then
+		kill -s KILL $ProcessesStillAlive 2> /dev/null
+	fi
+}
+
+# trap when SIGINT or SIGTERM are received with a custom function
+trap TermHandler SIGTERM SIGINT
 
 APPS=()
-APPS+=(google-chrome-stable)
-APPS+=(google-chrome)
+APPS+=(xdg-open)
 
 if [[ "${UNAMEOS}" =~ "Darwin" ]]; then
 	APPS+=(open)
@@ -31,41 +77,29 @@ fi
 
 ARGS=("${CMD}")
 
-if [[ "${CMD}" == *open ]]; then
-	ARGS+=(-a "google chrome")
-	ARGS+=(--new)
-	ARGS+=(-W)
-fi
-
-ARGS+=("file://${SCRIPT_DIR}/GPUDumpViewer.html")
-
-if [[ "${CMD}" == *open ]]; then
-	ARGS+=(--args)
-fi
-
-# --allow-file-access-from-files allow to load a file from a file:// webpage required for GPUDumpViewer.html to work.
-# --user-data-dir is required to force chrome to open a new instance so that --allow-file-access-from-files is honored.
-ARGS+=(--allow-file-access-from-files)
-ARGS+=(--new-window)
-ARGS+=(--incognito)
-ARGS+=("--user-data-dir=${CHROME_USER_DATA}")
+ARGS+=("$URL")
 
 echo "Executing:"
 echo
 
+echo "Starting simple webserver..."
+exec python3 -m http.server "$PORT" --bind "$ADDRESS" &
+P1=$!
+sleep 1
+
 echo "${ARGS[0]} \\"
 for ((i=1; i < ${#ARGS[@]}; i++ )); do
 	echo "  ${ARGS[$i]} \\";
 done
 echo
 
+# Start the browser now that the server is running
 "${ARGS[@]}"
 
+# Wait on the webserver - in general this will be killed by a Ctrl-C
+wait $P1
+
 echo
 echo "Closing ${CMD}..."
 
-if [[ -n "${CHROME_USER_DATA}" ]]; then
-	# Wait for 2s to shut down so that CHROME_USER_DATA can be deleted completely
-	sleep 2
-	rm -rf "${CHROME_USER_DATA}"
-fi
+popd
