Adds src color space management to correctly apply webgl's 2.2 gamma curve (or not) in the texture viewer.

diff --git a/Engine/Extras/GPUDumpViewer/GPUDumpViewer.html b/Engine/Extras/GPUDumpViewer/GPUDumpViewer.html
--- a/Engine/Extras/GPUDumpViewer/GPUDumpViewer.html
+++ b/Engine/Extras/GPUDumpViewer/GPUDumpViewer.html
@@ -1491,6 +1491,7 @@ class TextureView extends ResourceView
 		this.viewport_selected = false;
 		this.is_draging_canvas = false;
 		this.display_mode = 'Visualization';
+		this.src_color_space = 'sRGB/Rec709';
 	}
 
 	setup_html(parent_dom)
@@ -1544,14 +1545,19 @@ class TextureView extends ResourceView
 							---><input type="button" onclick="g_view.resource_view.change_pixel_scaling(this);" value="8:1"/><!--- 
 							---><input type="button" onclick="g_view.resource_view.change_pixel_scaling(this);" value="16:1"/>
 						</div>
-						<div class="button_list">
-							<input type="button" onclick="g_view.resource_view.copy_canvas_to_clipboard();" value="Copy to clipboard"/>
+						<div class="button_list" id="texture_visualization_change_src_color"><!---
+							---><span>Src color:</span><!---
+							---><input type="button" onclick="g_view.resource_view.change_src_color(this);" value="sRGB/Rec709" title="Performs no modifications to display to WebGL's sRGB backbuffer (default)."/><!---
+							---><input type="button" onclick="g_view.resource_view.change_src_color(this);" value="RGB Linear[0;1]" title="Applies a pow(display.rgb, 2.2); to output WebGL's sRGB backbuffer."/>
 						</div>
 						<div class="button_list" id="texture_visualization_change_display_modes"><!---
 							---><input type="button" onclick="g_view.resource_view.change_display_mode(this);" value="Visualization"/><!---
 							---><input type="button" onclick="g_view.resource_view.change_display_mode(this);" value="NaN"/><!---
 							---><input type="button" onclick="g_view.resource_view.change_display_mode(this);" value="Inf"/>
 						</div>
+						<div class="button_list">
+							<input type="button" onclick="g_view.resource_view.copy_canvas_to_clipboard();" value="Copy to clipboard"/>
+						</div>
 					</div>
 					<div id="canvas_viewport">
 						<canvas
@@ -1591,6 +1597,7 @@ class TextureView extends ResourceView
 		parent_dom.innerHTML = html;
 
 		update_value_selection(document.getElementById('texture_visualization_change_pixel_scaling'), `Fit`);
+		update_value_selection(document.getElementById('texture_visualization_change_src_color'), `sRGB/Rec709`);
 		update_value_selection(document.getElementById('texture_visualization_change_display_modes'), `Visualization`);
 
 		// Init WebGL 2.0
@@ -2365,6 +2372,14 @@ void main(void)
 			{
 				var name = prettify_subresource_unique_name(texture_view.subresource_version_info, texture_view.resource_desc);
 				text += `${is_output_resource ? 'Output' : 'Input'} resource:\n\t${name}\n`;
+				if (texture_view.display_mode == 'Visualization')
+				{
+					text += `\tSrc color: ${texture_view.src_color_space}\n`;
+				}
+				else
+				{
+					text += `\tSrc color: ${texture_view.display_mode}\n`;
+				}
 			}
 
 			if (texture_view.subresource_version_info['draw'] >= 0)
@@ -2815,6 +2830,25 @@ ${shader_display_assignment} = ${shader_display_operation};`;
 		update_value_selection(document.getElementById('texture_visualization_change_pixel_scaling'), pixel_scaling_button.value);
 	}
 
+	change_src_color(new_src_color_space)
+	{
+		if (!this.is_ready)
+		{
+			return;
+		}
+		
+		if (new_src_color_space.value == this.src_color_space)
+		{
+			return;
+		}
+
+		this.src_color_space = new_src_color_space.value;
+
+		update_value_selection(document.getElementById('texture_visualization_change_src_color'), this.src_color_space);
+
+		this.refresh_texture_visualization();
+	}
+
 	change_display_mode(new_mode)
 	{
 		if (!this.is_ready)
@@ -3146,7 +3180,22 @@ void main(void) {
 		else
 		{
 			var user_frag_code = document.getElementById('texture_visualization_code_input').value;
-			shader_program = this.compile_screen_shader_program(this.subresource_desc['webgl_pixel_shader_private'] + 'void main(void) { display = vec4(0.0, 0.0, 0.0, 0.0);\n' + user_frag_code + `\n}`);
+
+			var final_frag_code = this.subresource_desc['webgl_pixel_shader_private'] + 'void main(void) { display = vec4(0.0, 0.0, 0.0, 0.0);\n';
+			final_frag_code += user_frag_code;
+			final_frag_code += '\n';
+			if (this.src_color_space == 'sRGB/Rec709')
+			{
+				// NOP because src gamma = webgl display gamma.
+			}
+			else if (this.src_color_space == 'RGB Linear[0;1]')
+			{
+				// Apply webgl's 2.2 gamma curve on the display color so linear gradients look linear in web browser.
+				final_frag_code += 'display.rgb = pow(display.rgb, vec3(2.2));';
+			}
+			final_frag_code += '}';
+
+			shader_program = this.compile_screen_shader_program(final_frag_code);
 
 			if (!shader_program)
 			{
@@ -4638,6 +4687,12 @@ input::placeholder
 	color: var(--error-txt-color);
 }
 
+.button_list span
+{
+	margin-left: 5px;
+	margin-right: 3px;
+}
+
 
 /* ----------------------------------------------------- top bar */
 
