Adds display channel selectors in the texture viewer

diff --git a/Engine/Extras/GPUDumpViewer/GPUDumpViewer.html b/Engine/Extras/GPUDumpViewer/GPUDumpViewer.html
--- a/Engine/Extras/GPUDumpViewer/GPUDumpViewer.html
+++ b/Engine/Extras/GPUDumpViewer/GPUDumpViewer.html
@@ -1490,8 +1490,9 @@ class TextureView extends ResourceView
 		this.pixel_scaling = 0;
 		this.viewport_selected = false;
 		this.is_draging_canvas = false;
-		this.display_mode = 'Visualization';
+		this.src_channels = 'RGB';
 		this.src_color_space = 'sRGB/Rec709';
+		this.display_mode = 'Visualization';
 	}
 
 	setup_html(parent_dom)
@@ -1545,15 +1546,22 @@ class TextureView extends ResourceView
 							---><input type="button" onclick="g_view.resource_view.change_pixel_scaling(this);" value="8:1"/><!--- 
 							---><input type="button" onclick="g_view.resource_view.change_pixel_scaling(this);" value="16:1"/>
 						</div>
+						<div class="button_list" id="texture_visualization_change_src_channels"><!---
+							---><input type="button" onclick="g_view.resource_view.change_src_channels(this);" value="RGB" title="Display the WebGL's visualization shader's display.rgb channels (default)."/><!---
+							---><input type="button" onclick="g_view.resource_view.change_src_channels(this);" value="R" title="Display the WebGL's visualization shader's display.r channel."/><!---
+							---><input type="button" onclick="g_view.resource_view.change_src_channels(this);" value="G" title="Display the WebGL's visualization shader's display.g channel."/><!---
+							---><input type="button" onclick="g_view.resource_view.change_src_channels(this);" value="B" title="Display the WebGL's visualization shader's display.b channel."/><!---
+							---><input type="button" onclick="g_view.resource_view.change_src_channels(this);" value="A" title="Display the WebGL's visualization shader's display.a channel."/>
+						</div>
 						<div class="button_list" id="texture_visualization_change_src_color"><!---
 							---><span>Src color:</span><!---
 							---><input type="button" onclick="g_view.resource_view.change_src_color(this);" value="sRGB/Rec709" title="Performs no modifications to display to WebGL's sRGB backbuffer (default)."/><!---
 							---><input type="button" onclick="g_view.resource_view.change_src_color(this);" value="RGB Linear[0;1]" title="Applies a pow(display.rgb, 2.2); to output WebGL's sRGB backbuffer."/>
 						</div>
 						<div class="button_list" id="texture_visualization_change_display_modes"><!---
-							---><input type="button" onclick="g_view.resource_view.change_display_mode(this);" value="Visualization"/><!---
-							---><input type="button" onclick="g_view.resource_view.change_display_mode(this);" value="NaN"/><!---
-							---><input type="button" onclick="g_view.resource_view.change_display_mode(this);" value="Inf"/>
+							---><input type="button" onclick="g_view.resource_view.change_display_mode(this);" value="Visualization" title="Display the WebGL's visualization shader's display.rgb."/><!---
+							---><input type="button" onclick="g_view.resource_view.change_display_mode(this);" value="NaN" title="Display any NaN in the raw texture."/><!---
+							---><input type="button" onclick="g_view.resource_view.change_display_mode(this);" value="Inf" title="Display any Inf in the raw texture."/>
 						</div>
 						<div class="button_list">
 							<input type="button" onclick="g_view.resource_view.copy_canvas_to_clipboard();" value="Copy to clipboard"/>
@@ -1597,8 +1605,9 @@ class TextureView extends ResourceView
 		parent_dom.innerHTML = html;
 
 		update_value_selection(document.getElementById('texture_visualization_change_pixel_scaling'), `Fit`);
-		update_value_selection(document.getElementById('texture_visualization_change_src_color'), `sRGB/Rec709`);
-		update_value_selection(document.getElementById('texture_visualization_change_display_modes'), `Visualization`);
+		update_value_selection(document.getElementById('texture_visualization_change_src_channels'), this.src_channels);
+		update_value_selection(document.getElementById('texture_visualization_change_src_color'), this.src_color_space);
+		update_value_selection(document.getElementById('texture_visualization_change_display_modes'), this.display_mode);
 
 		// Init WebGL 2.0
 		this.init_gl();
@@ -2374,7 +2383,9 @@ void main(void)
 				text += `${is_output_resource ? 'Output' : 'Input'} resource:\n\t${name}\n`;
 				if (texture_view.display_mode == 'Visualization')
 				{
-					text += `\tSrc color: ${texture_view.src_color_space}\n`;
+					text += `\tDisplay shader: \n${g_shader_code_dict[texture_view.shader_code_saving_key]}\n`;
+					text += `\tDisplay channels: ${texture_view.src_channels}\n`;
+					text += `\tDisplay src color: ${texture_view.src_color_space}\n`;
 				}
 				else
 				{
@@ -2830,6 +2841,25 @@ ${shader_display_assignment} = ${shader_display_operation};`;
 		update_value_selection(document.getElementById('texture_visualization_change_pixel_scaling'), pixel_scaling_button.value);
 	}
 
+	change_src_channels(new_src_channels)
+	{
+		if (!this.is_ready)
+		{
+			return;
+		}
+		
+		if (new_src_channels.value == this.src_channels)
+		{
+			return;
+		}
+
+		this.src_channels = new_src_channels.value;
+
+		update_value_selection(document.getElementById('texture_visualization_change_src_channels'), this.src_channels);
+
+		this.refresh_texture_visualization();
+	}
+
 	change_src_color(new_src_color_space)
 	{
 		if (!this.is_ready)
@@ -3184,6 +3214,24 @@ void main(void) {
 			var final_frag_code = this.subresource_desc['webgl_pixel_shader_private'] + 'void main(void) { display = vec4(0.0, 0.0, 0.0, 0.0);\n';
 			final_frag_code += user_frag_code;
 			final_frag_code += '\n';
+
+			if (this.src_channels == 'R')
+			{
+				final_frag_code += 'display.rgb = vec3(display.r, 0.0, 0.0);';
+			}
+			else if (this.src_channels == 'G')
+			{
+				final_frag_code += 'display.rgb = vec3(0.0, display.g, 0.0);';
+			}
+			else if (this.src_channels == 'B')
+			{
+				final_frag_code += 'display.rgb = vec3(0.0, 0.0, display.b);';
+			}
+			else if (this.src_channels == 'A')
+			{
+				final_frag_code += 'display.rgb = display.aaa;';
+			}
+
 			if (this.src_color_space == 'sRGB/Rec709')
 			{
 				// NOP because src gamma = webgl display gamma.
