Fixes log viewer's scroll bars when there is very long log lines.

diff --git a/Engine/Extras/GPUDumpViewer/GPUDumpViewer.html b/Engine/Extras/GPUDumpViewer/GPUDumpViewer.html
--- a/Engine/Extras/GPUDumpViewer/GPUDumpViewer.html
+++ b/Engine/Extras/GPUDumpViewer/GPUDumpViewer.html
@@ -9,7 +9,10 @@
 
 const k_current_dir = "./";
 
-const k_scroll_width = 8;
+const k_style_scroll_width = 8;
+const k_style_padding_size = 10;
+const k_style_main_div_margin = 15;
+
 const k_null_json_ptr = "0000000000000000";
 
 
@@ -1415,7 +1418,7 @@ class LogView extends IView
 	{
 		parent_dom.innerHTML = `
 			<div class="main_div">
-				<div class="pass_title">Log</div>
+				<div class="pass_title">${g_infos['LogFilename']}</div>
 				<div class="selection_list_search">
 					<input type="search" id="log_search_input" oninput="g_view.refresh_log();" onchange="g_view.refresh_log();" style="width: 100%;" placeholder="Search log..." />
 				</div>
@@ -1461,6 +1464,13 @@ class LogView extends IView
 		return [`display_log();`];
 	}
 
+	resize(ctx)
+	{
+		document.getElementById('log_pannel').style.width = `${ctx.width - 2 * k_style_padding_size}px`;
+		document.getElementById('log_pannel').style.height = `${ctx.height - 2 * k_style_padding_size - 75}px`;
+		document.getElementById('log_pannel').style.overflow = `scroll`;
+	}
+
 	release()
 	{
 		this.log = null;
@@ -1912,7 +1922,7 @@ class TextureView extends ResourceView
 		// Find new pixel scaling to scale the viewport to.
 		var new_pixel_scaling = 0;
 		{
-			var viewport_width = (px_string_to_int(document.getElementById('canvas_viewport').style.width) - k_scroll_width) * dpi;
+			var viewport_width = (px_string_to_int(document.getElementById('canvas_viewport').style.width) - k_style_scroll_width) * dpi;
 
 			var min_pixel_scaling = Math.ceil(viewport_width / subresource_extent['x']);
 			if (min_pixel_scaling <= 1)
@@ -2093,8 +2103,8 @@ class TextureView extends ResourceView
 	resize(ctx)
 	{
 		var dpi = window.devicePixelRatio || 1;
-		var viewport_width = ctx.width - 100;
-		var viewport_width_no_scroll = viewport_width - k_scroll_width;
+		var viewport_width = ctx.width - 2 * k_style_padding_size - 2;
+		var viewport_width_no_scroll = viewport_width - k_style_scroll_width;
 
 		var subresource_extent = this.get_subresource_extent();
 		var canvas_rendering_res_x = subresource_extent['x'];
@@ -2103,8 +2113,8 @@ class TextureView extends ResourceView
 		var canvas_outter = document.getElementById('canvas_outter');
 		var canvas_viewport = document.getElementById('canvas_viewport');
 		canvas_outter.style.width = `${viewport_width}px`;
-		canvas_viewport.style.width = `${viewport_width_no_scroll + k_scroll_width}px`;
-		canvas_viewport.style.height = `${viewport_width_no_scroll * canvas_rendering_res_y / canvas_rendering_res_x + k_scroll_width}px`;
+		canvas_viewport.style.width = `${viewport_width_no_scroll + k_style_scroll_width}px`;
+		canvas_viewport.style.height = `${viewport_width_no_scroll * canvas_rendering_res_y / canvas_rendering_res_x + k_style_scroll_width}px`;
 
 		this.viewport_width = viewport_width;
 		if (this.is_ready)
@@ -2118,12 +2128,12 @@ class TextureView extends ResourceView
 		console.assert(this.is_ready);
 		var dpi = window.devicePixelRatio || 1;
 		var viewport_width = this.viewport_width;
-		var viewport_width_no_scroll = viewport_width - k_scroll_width;
+		var viewport_width_no_scroll = viewport_width - k_style_scroll_width;
 
 		var canvas_viewport = document.getElementById('canvas_viewport');
 		var canvas_viewport_rect = canvas_viewport.getBoundingClientRect();
 		var viewport_height = px_string_to_int(canvas_viewport.style.height);
-		var viewport_height_no_scroll = viewport_height - k_scroll_width;
+		var viewport_height_no_scroll = viewport_height - k_style_scroll_width;
 
 		var subresource_extent = this.get_subresource_extent();
 		var canvas_rendering_res_x = subresource_extent['x'];
@@ -3877,7 +3887,7 @@ class GenericBufferView extends ResourceView
 
 		{
 			var min_row_count = 10;
-			document.getElementById('buffer_viewport').style.height = `${k_scroll_width + Math.min(Math.max(row_count + 1, min_row_count), k_buffer_max_display_row_count) * k_buffer_pixel_per_row}px`;
+			document.getElementById('buffer_viewport').style.height = `${k_style_scroll_width + Math.min(Math.max(row_count + 1, min_row_count), k_buffer_max_display_row_count) * k_buffer_pixel_per_row}px`;
 		}
 	}
 
@@ -4243,38 +4253,50 @@ function init_page()
 
 function onresize_body()
 {
-	var w = window.innerWidth;
-	var h = window.innerHeight;
+	var window_width = window.innerWidth;
+	var window_height = window.innerHeight;
 
-	var top_h = 40;
+	const top_h = 45;
 
-	document.getElementById('onresize_body_top_bar').style.width = `${w}px`;
-	document.getElementById('onresize_body_top_bar').style.height = `${top_h}px`;
+	document.getElementById('onresize_body_table').style.width = `${window_width}px`;
+	document.getElementById('onresize_body_table').style.height = `${window_height}px`;
 
-	document.getElementById('onresize_body_left').style.width = `${Math.floor(w * 0.25)}px`;
-	document.getElementById('onresize_body_left').style.height = `${h - top_h}px`;
+	const top_height = 45;
 
-	document.getElementById('onresize_body_right').style.width = `${w - Math.floor(w * 0.25)}px`;
-	document.getElementById('onresize_body_right').style.height = `${h - top_h}px`;
+	var body_left_width = Math.floor(window_width * 0.25);
+	var body_right_width = window_width - body_left_width;
+	var body_height = window_height - top_height;
+	{
+		document.getElementById('onresize_body_top_bar').style.width = `${window_width}px`;
+		document.getElementById('onresize_body_top_bar').style.height = `${top_height}px`;
+		document.getElementById('onresize_body_top_bar').style.overflow = `hidden`;
+
+		document.getElementById('onresize_body_left').style.width = `${body_left_width}px`;
+		document.getElementById('onresize_body_left').style.height = `${body_height}px`;
 
-	document.getElementById('onresize_body_right').style.width = `${w - Math.floor(w * 0.25)}px`;
-	document.getElementById('onresize_body_right').style.height = `${h - top_h}px`;
+		document.getElementById('onresize_body_right').style.width = `${body_right_width}px`;
+		document.getElementById('onresize_body_right').style.height = `${body_height}px`;
+	}
 
 	{
-		var search_h = (30 + 2 * 5) * 2;
-		document.getElementById('pass_hierarchy').style.width = `${Math.floor(w * 0.25) - 10 * 2}px`;
-		document.getElementById('pass_hierarchy').style.height = `${h - top_h - search_h - 10 * 2}px`;
+		const search_h = (30 + 2 * 5) * 2;
+
+		document.getElementById('pass_hierarchy').style.width = `${body_left_width - k_style_padding_size * 2}px`;
+		document.getElementById('pass_hierarchy').style.height = `${body_height - search_h - k_style_padding_size * 2 - k_style_scroll_width}px`;
+		document.getElementById('pass_hierarchy').style.overflow = `scroll`;
 	}
 
 	{
-		//document.getElementById('main_right_pannel').style.width = document.getElementById('onresize_body_right').style.width;
-		document.getElementById('main_right_pannel').style.height = `${h - top_h - 10 * 2}px`;
+		document.getElementById('main_right_pannel').style.width = `${body_right_width}px`;
+		document.getElementById('main_right_pannel').style.height = `${body_height}px`;
+		document.getElementById('main_right_pannel').style.overflow = `scroll`;
 	}
 
 	if (g_view !== null)
 	{
 		var ctx = {};
-		ctx.width = w - Math.floor(w * 0.25);
+		ctx.width = body_right_width - k_style_scroll_width;
+		ctx.height = body_height - k_style_scroll_width;
 		g_view.resize(ctx);
 	}
 }
@@ -4537,6 +4559,11 @@ input::placeholder
 	margin-bottom: 20px;
 }
 
+#main_right_pannel .main_div:last-child
+{
+	margin-bottom: 5px;
+}
+
 
 /* ----------------------------------------------------- Selection list */
 
@@ -4763,13 +4790,6 @@ input::placeholder
 
 /* ----------------------------------------------------- main_right_pannel */
 
-#main_right_pannel
-{
-	width: auto;
-	overflow-x: auto;
-	overflow-y: scroll;
-}
-
 #main_right_pannel .pass_title
 {
 	font-size: 20;
