Adds buttons in DumpGPU viewer to move forward/backward between frames when viewing a resource

diff --git a/Engine/Extras/GPUDumpViewer/GPUDumpViewer.html b/Engine/Extras/GPUDumpViewer/GPUDumpViewer.html
--- a/Engine/Extras/GPUDumpViewer/GPUDumpViewer.html
+++ b/Engine/Extras/GPUDumpViewer/GPUDumpViewer.html
@@ -494,6 +494,136 @@ function test_decode_float()
 }
 
 
+// -------------------------------------------------------------------- PASS ANALYSIS
+
+var g_analysis = {};
+
+function get_frame_number(pass_data)
+{
+	var pass_event_scopes = pass_data['ParentEventScopes'];
+	var frame_event_scope_name = pass_event_scopes[pass_event_scopes.length - 1];
+
+	const regex = new RegExp("Frame (?<FrameNumber>\\d+).*");
+	//const regex = /Frame (?<FrameNumber>\d+) .*"/;
+
+	var found = frame_event_scope_name.match(regex);
+	return Number(found.groups['FrameNumber']);
+}
+
+function compare_draw_events(ref_draw_event, tested_draw_event)
+{
+	return ref_draw_event == tested_draw_event;
+}
+
+function is_pass_output_resource(pass_data, subresource_version_info)
+{
+	var subresource_unique_name = get_subresource_unique_name(subresource_version_info);
+
+	var is_output_resource = null;
+	if (pass_data['InputResources'].includes(subresource_unique_name))
+	{
+		is_output_resource = false;
+	}
+	else if (pass_data['OutputResources'].includes(subresource_unique_name))
+	{
+		is_output_resource = true;
+	}
+	console.assert(is_output_resource !== null);
+	return is_output_resource;
+}
+
+function find_similar_pass_id_in_frame(current_frame_pass_data, new_frame_id)
+{
+	var current_pass_event_scopes = current_frame_pass_data['ParentEventScopes'];
+
+	for (var pass_id in g_passes)
+	{
+		var pass_data = g_passes[pass_id];
+		if (!compare_draw_events(pass_data['EventName'], current_frame_pass_data['EventName']))
+		{
+			continue;
+		}
+
+		var frame_id = get_frame_number(pass_data);
+		if (frame_id != new_frame_id)
+		{
+			continue;
+		}
+
+		var pass_event_scopes = pass_data['ParentEventScopes'];
+		if (pass_event_scopes.length != current_pass_event_scopes.length)
+		{
+			continue;
+		}
+		else
+		{
+			var is_same_parent = true;
+			for (var parent_scope_id = 0; parent_scope_id < pass_event_scopes.length - 1; parent_scope_id++)
+			{
+				is_same_parent = is_same_parent && compare_draw_events(current_pass_event_scopes[parent_scope_id], pass_event_scopes[parent_scope_id]);
+			}
+			if (!is_same_parent)
+			{
+				continue;
+			}
+		}
+
+		return Number(pass_id);
+	}
+
+	return null;
+}
+
+function find_similar_resource_in_pass(current_frame_pass_data, current_frame_resource_desc, current_frame_subresource_info, new_pass_data)
+{
+	var is_output_resource = is_pass_output_resource(current_frame_pass_data, current_frame_subresource_info);
+
+	var new_pass_resource_list = is_output_resource ? new_pass_data['OutputResources'] : new_pass_data['InputResources'];
+
+	for (var new_pass_subresource_unique_name of new_pass_resource_list)
+	{
+		var new_pass_subresource_info = parse_subresource_unique_name(new_pass_subresource_unique_name);
+		if (new_pass_subresource_info['subresource'] !== current_frame_subresource_info['subresource'])
+		{
+			continue;
+		}
+		if (new_pass_subresource_info['array_slice'] !== current_frame_subresource_info['array_slice'])
+		{
+			continue;
+		}
+
+		var new_pass_resource_desc = get_resource_desc(new_pass_subresource_info['resource']);
+		if (new_pass_resource_desc['Name'] !== current_frame_resource_desc['Name'])
+		{
+			continue;
+		}
+
+		return new_pass_subresource_info;
+	}
+
+	return null;
+}
+
+function analyses_passes()
+{
+	var frames = [];
+
+	for (var pass_id in g_passes)
+	{
+		var pass_data = g_passes[pass_id];
+		
+		var frame_id = get_frame_number(pass_data);
+
+		if (!frames.includes(frame_id))
+		{
+			frames.push(frame_id);
+		}
+	}
+
+	g_analysis['FrameList'] = frames;
+}
+
+
 // -------------------------------------------------------------------- DISPLAY PASS
 
 function display_pass_hierarchy()
@@ -613,6 +743,7 @@ class PassView extends IView
 
 		parent_dom.innerHTML = `
 			<div class="pass_title main_div">
+				<div id="pass_frame_list" style="display:none;"></div>
 				${this.pass_data['EventName']}
 				<div class="button_list" style="display:inline-block;"><a href="#display_pass_parameters(${this.pass_id});">PassParameters</a></div>
 			</div>
@@ -675,6 +806,7 @@ class PassView extends IView
 
 	refresh_all_lists()
 	{
+		this.refresh_frames_lists();
 		this.refresh_input_resource_list();
 		this.refresh_output_resource_list();
 		this.refresh_draw_resource_list();
@@ -682,6 +814,65 @@ class PassView extends IView
 		this.refresh_resource_reading_list();
 	}
 
+	refresh_frames_lists()
+	{
+		if (g_analysis['FrameList'].length == 1)
+		{
+			return;
+		}
+
+		if (this.resource_view === null)
+		{
+			document.getElementById('pass_frame_list').style.display = 'none';
+			return;
+		}
+
+		var pass_data = this.pass_data;
+		var subresource_version_info = this.resource_view.subresource_version_info;
+
+		var current_frame_number = get_frame_number(this.pass_data);
+		var prev_frame_number = g_analysis['FrameList'].includes(current_frame_number - 1) ? current_frame_number - 1 : g_analysis['FrameList'][g_analysis['FrameList'].length - 1];
+		var next_frame_number = g_analysis['FrameList'].includes(current_frame_number + 1) ? current_frame_number + 1 : g_analysis['FrameList'][0];
+
+		var prev_frame_pass_id = find_similar_pass_id_in_frame(this.pass_data, prev_frame_number);
+		var next_frame_pass_id = find_similar_pass_id_in_frame(this.pass_data, next_frame_number);
+
+		var prev_frame_href = '';
+		var next_frame_href = '';
+		if (this.resource_view.resource_desc['Desc'] == 'FRDGParameterStruct')
+		{
+			prev_frame_href = `#display_pass_parameters(${prev_frame_pass_id});`;
+			next_frame_href = `#display_pass_parameters(${next_frame_pass_id});`;
+		}
+		else
+		{
+			var prev_frame_pass_data = g_passes[prev_frame_pass_id];
+			var next_frame_pass_data = g_passes[next_frame_pass_id];
+
+			var is_output_resource = is_pass_output_resource(this.pass_data, this.resource_view.subresource_version_info);
+			var prev_frame_subresource_info = find_similar_resource_in_pass(this.pass_data, this.resource_view.resource_desc, this.resource_view.subresource_version_info, prev_frame_pass_data);
+			var next_frame_subresource_info = find_similar_resource_in_pass(this.pass_data, this.resource_view.resource_desc, this.resource_view.subresource_version_info, next_frame_pass_data);
+
+			var prev_frame_href = '';
+			var next_frame_href = '';
+			if (is_output_resource)
+			{
+				prev_frame_href = `#display_output_resource(${prev_frame_pass_id},'${get_subresource_unique_name(prev_frame_subresource_info)}');`;
+				next_frame_href = `#display_output_resource(${next_frame_pass_id},'${get_subresource_unique_name(next_frame_subresource_info)}');`;
+			}
+			else
+			{
+				prev_frame_href = `#display_input_resource(${prev_frame_pass_id},'${get_subresource_unique_name(prev_frame_subresource_info)}');`;
+				next_frame_href = `#display_input_resource(${next_frame_pass_id},'${get_subresource_unique_name(next_frame_subresource_info)}');`;
+			}
+		}
+
+		document.getElementById('pass_frame_list').style.display = 'inline-block';
+		document.getElementById('pass_frame_list').innerHTML = `
+			<div class="button_list" style="display:inline-block;"><a href="${prev_frame_href}">-</a><a title="Id of the current frame">${current_frame_number}</a><a href="${next_frame_href}">+</a></div>
+			`;
+	}
+
 	refresh_input_resource_list()
 	{
 		var display_list = [];
@@ -755,7 +946,7 @@ class PassView extends IView
 		var subresource_unique_name = '';
 		if (this.resource_view)
 		{
-			is_output_resource = this.pass_data['OutputResources'].includes(get_subresource_unique_name(this.resource_view.subresource_version_info));
+			is_output_resource = is_pass_output_resource(this.pass_data, this.resource_view.subresource_version_info);
 			subresource_unique_name = get_subresource_unique_name(this.resource_view.subresource_version_info);
 		}
 
@@ -2383,7 +2574,7 @@ void main(void)
 		this.canvas.toBlob(function(image_blob) {
 			var pass_data = g_view.pass_data;
 
-			var is_output_resource = pass_data['OutputResources'].includes(get_subresource_unique_name(texture_view.subresource_version_info));
+			var is_output_resource = is_pass_output_resource(pass_data, texture_view.subresource_version_info);
 
 			var text = '';
 			text += `Dump:\n\t${g_infos['Project']} - ${g_infos['Platform']} - ${g_infos['RHI']} - ${g_infos['BuildVersion']} - ${g_infos['DumpTime']}\n\t${get_current_navigation()}\n`
@@ -4200,6 +4391,11 @@ function init_page()
 		add_console_event('log', `Viewer init ok`);
 	}
 
+	// Analyses the capture.
+	{
+		analyses_passes();
+	}
+
 	init_top();
 	init_top_buttons();
 	display_pass_hierarchy();
@@ -4718,10 +4914,15 @@ input::placeholder
 	padding: 3px 10px;
 	background-color: var(--button-bg-color);
 	color: var(--button-txt-color);
-	cursor: pointer;
 	border-radius: 0;
 }
 
+.button_list a[href],
+.button_list input
+{
+	cursor: pointer;
+}
+
 .button_list a:first-child,
 .button_list input:first-child
 {
@@ -4736,20 +4937,20 @@ input::placeholder
 	border-bottom-right-radius: var(--button-border-radius);
 }
 
-.button_list a:not(.match_location_hash):not(:active):hover,
+.button_list a[href]:not(.match_location_hash):not(:active):hover,
 .button_list input:not(.match_value):not(:active):hover
 {
 	background-color: var(--button-bg-color-hover);
 }
 
-.button_list a:not(.match_location_hash):not(.error):not(:active):hover,
+.button_list a[href]:not(.match_location_hash):not(.error):not(:active):hover,
 .button_list input:not(.match_value):not(:active):hover
 {
 	color: var(--button-txt-color-hover);
 }
 
-.button_list a:active,
-.button_list a.match_location_hash,
+.button_list a[href]:active,
+.button_list a[href].match_location_hash,
 .button_list input:active,
 .button_list input.match_value
 {
@@ -4757,8 +4958,8 @@ input::placeholder
 	color: var(--button-txt-color-selected);
 }
 
-.button_list a.error,
-.button_list a.error:active
+.button_list a[href].error,
+.button_list a[href].error:active
 {
 	color: var(--error-txt-color);
 }
