Implements a Generic3DStructureView to be able to implements 3D structure viewers

diff --git a/Engine/Extras/GPUDumpViewer/GPUDumpViewer.html b/Engine/Extras/GPUDumpViewer/GPUDumpViewer.html
--- a/Engine/Extras/GPUDumpViewer/GPUDumpViewer.html
+++ b/Engine/Extras/GPUDumpViewer/GPUDumpViewer.html
@@ -1674,6 +1674,231 @@ function display_log()
 }
 
 
+
+// -------------------------------------------------------------------- WEBGL UTILS TEXTURE
+
+function gl_create_vertex_buffer(gl, vertices)
+{
+	var webgl_vertex_buffer = gl.createBuffer();
+	gl.bindBuffer(gl.ARRAY_BUFFER, webgl_vertex_buffer);
+	gl.bufferData(gl.ARRAY_BUFFER, vertices, gl.STATIC_DRAW);
+	gl.bindBuffer(gl.ARRAY_BUFFER, null);
+	return webgl_vertex_buffer;
+}
+
+function gl_create_index_buffer(gl, indices)
+{
+	var webgl_index_buffer = gl.createBuffer();
+	gl.bindBuffer(gl.ELEMENT_ARRAY_BUFFER, webgl_index_buffer);
+	gl.bufferData(gl.ELEMENT_ARRAY_BUFFER, indices, gl.STATIC_DRAW);
+	gl.bindBuffer(gl.ELEMENT_ARRAY_BUFFER, null);
+	return webgl_index_buffer;
+}
+
+function gl_create_shader(gl, shader_type, shader_code)
+{
+	const shader_std = `#version 300 es
+
+precision highp float;
+precision highp int;
+precision highp sampler2D;
+precision highp usampler2D;
+
+struct FloatInfos
+{
+	bool is_denormal;
+	bool is_infinity;
+	bool is_nan;
+};
+
+FloatInfos decodeFloat(float x)
+{
+	const uint mantissa_bit_count = 23u;
+	const uint exp_bit_count = 8u;
+
+	uint raw = floatBitsToUint(x);
+
+	uint sign_bit      = (raw >> (mantissa_bit_count + exp_bit_count)) & 0x1u;
+	uint mantissa_bits = (raw >> 0u)                     & ((0x1u << mantissa_bit_count) - 1u);
+	uint exp_bits      = (raw >> mantissa_bit_count)     & ((0x1u << exp_bit_count) - 1u);
+
+	bool is_max_exp = exp_bits == ((0x1u << exp_bit_count) - 1u);
+
+	FloatInfos infos;
+	infos.is_denormal = exp_bits == 0u;
+	infos.is_infinity = is_max_exp && mantissa_bits == 0u;
+	infos.is_nan = is_max_exp && mantissa_bits != 0u;
+
+	return infos;
+}
+
+#define isnan(x) decodeFloat(x).is_nan
+#define isinf(x) decodeFloat(x).is_infinity
+#define isdenorm(x) decodeFloat(x).is_denormal
+
+`;
+	
+	var gl_shader = gl.createShader(shader_type);
+	gl.shaderSource(gl_shader, shader_std + shader_code);
+	gl.compileShader(gl_shader);
+
+	var compilation_status = gl.getShaderParameter(gl_shader, gl.COMPILE_STATUS);
+	if (!compilation_status)
+	{
+		var compilation_log = gl.getShaderInfoLog(gl_shader);
+		console.error('Shader compiler log: ' + compilation_log);
+		return compilation_log;
+	}
+	return gl_shader;
+}
+
+function gl_create_shader_program(gl, vert_shader, frag_shader)
+{
+	var gl_vert_shader = null;
+	if (typeof vert_shader === 'string')
+	{
+		gl_vert_shader = gl_create_shader(gl, gl.VERTEX_SHADER, vert_shader);
+		if (typeof gl_vert_shader === 'string')
+		{
+			return gl_vert_shader;
+		}
+	}
+	else if (typeof vert_shader === 'WebGLShader')
+	{
+		gl_vert_shader = vert_shader;
+	}
+	else
+	{
+		console.error('wrong vert_shader');
+		return null;
+	}
+
+	var gl_frag_shader = null;
+	if (typeof frag_shader === 'string')
+	{
+		gl_frag_shader = gl_create_shader(gl, gl.FRAGMENT_SHADER, frag_shader);
+		if (typeof gl_frag_shader === 'string')
+		{
+			return gl_frag_shader;
+		}
+	}
+	else if (typeof frag_shader === 'WebGLShader')
+	{
+		gl_frag_shader = frag_shader;
+	}
+	else
+	{
+		console.error('wrong vert_shader');
+		return null;
+	}
+
+	var shader_program = gl.createProgram();
+	gl.attachShader(shader_program, gl_vert_shader);
+	gl.attachShader(shader_program, gl_frag_shader);
+	gl.linkProgram(shader_program);
+
+	var link_status = gl.getProgramParameter(shader_program, gl.LINK_STATUS);
+	if (!link_status)
+	{
+		var link_log = gl.getProgramInfoLog(shader_program);
+		console.error('Program compiler log: ' + link_log);
+		return null;
+	}
+	return shader_program;
+}
+
+function gl_set_uniform_mat4(gl, program, name, matrix)
+{
+	gl.uniformMatrix4fv(
+		gl.getUniformLocation(program, name),
+		false,
+		new Float32Array(matrix));
+}
+
+
+// -------------------------------------------------------------------- MATH UTILS
+
+function create_null_matrix()
+{
+	return new Array(4 * 4).fill(0.0);
+}
+
+function create_identity_matrix()
+{
+	var matrix = create_null_matrix();
+	matrix[0 + 0 * 4] = 1.0;
+	matrix[1 + 1 * 4] = 1.0;
+	matrix[2 + 2 * 4] = 1.0;
+	matrix[3 + 3 * 4] = 1.0;
+	return matrix;
+}
+
+function create_view_matrix(camera_pos, camera_dir)
+{
+	var camera_pos_x = camera_pos[0];
+	var camera_pos_y = camera_pos[1];
+	var camera_pos_z = camera_pos[2];
+
+	var camera_dir_len = Math.sqrt(
+		camera_dir[0] * camera_dir[0] +
+		camera_dir[1] * camera_dir[1] +
+		camera_dir[2] * camera_dir[2]);
+	var camera_dir_x = camera_dir[0] / camera_dir_len;
+	var camera_dir_y = camera_dir[1] / camera_dir_len;
+	var camera_dir_z = camera_dir[2] / camera_dir_len;
+
+	var horizon_dir_len = Math.sqrt(camera_dir_x * camera_dir_x + camera_dir_y * camera_dir_y);
+	var horizon_dir_x = camera_dir_y / horizon_dir_len;
+	var horizon_dir_y = -camera_dir_x / horizon_dir_len;
+	var horizon_dir_z = 0.0;
+
+	var vertical_dir_x = - camera_dir_y * horizon_dir_z + camera_dir_z * horizon_dir_y;
+	var vertical_dir_y = - camera_dir_z * horizon_dir_x + camera_dir_x * horizon_dir_z;
+	var vertical_dir_z = - camera_dir_x * horizon_dir_y + camera_dir_y * horizon_dir_x;
+
+	// column major
+	var matrix = create_identity_matrix();
+	matrix[0 + 0 * 4] = horizon_dir_x;
+	matrix[0 + 1 * 4] = horizon_dir_y;
+	matrix[0 + 2 * 4] = horizon_dir_z;
+	matrix[0 + 3 * 4] = -(horizon_dir_x * camera_pos_x + horizon_dir_y * camera_pos_y + horizon_dir_z * camera_pos_z);
+
+	matrix[1 + 0 * 4] = vertical_dir_x;
+	matrix[1 + 1 * 4] = vertical_dir_y;
+	matrix[1 + 2 * 4] = vertical_dir_z;
+	matrix[1 + 3 * 4] = -(vertical_dir_x * camera_pos_x + vertical_dir_y * camera_pos_y + vertical_dir_z * camera_pos_z);
+
+	matrix[2 + 0 * 4] = camera_dir_x;
+	matrix[2 + 1 * 4] = camera_dir_y;
+	matrix[2 + 2 * 4] = camera_dir_z;
+	matrix[2 + 3 * 4] = -(camera_dir_x * camera_pos_x + camera_dir_y * camera_pos_y + camera_dir_z * camera_pos_z);
+
+	return matrix;
+}
+
+// Matches TPerspectiveMatrix
+function create_projection_matrix_reverse_z_persepective(half_fov, aspect_ratio, clipping_plane)
+{
+	const z_precision = 0.0;
+
+	var half_fov_x = half_fov * Math.PI / 180.0;
+	var half_fov_y = half_fov_x;
+	var mult_fov_x = 1.0;
+	var mult_fov_y = aspect_ratio;
+	var min_z = clipping_plane;
+	var max_z = clipping_plane;
+
+	// column major
+	var matrix = create_null_matrix();
+	matrix[0 + 0 * 4] = mult_fov_x / Math.tan(half_fov_x);
+	matrix[1 + 1 * 4] = mult_fov_y / Math.tan(half_fov_y);
+	matrix[2 + 3 * 4] = min_z;
+	matrix[3 + 2 * 4] = 1.0;
+
+	return matrix;
+}
+
+
 // -------------------------------------------------------------------- DISPLAY TEXTURE
 
 var k_exotic_raw_channel_bit_depth = 0;
@@ -2450,29 +2675,16 @@ class TextureView extends ResourceView
 		// Set unpacking alignement to 1 to allow uploading texture size not multple of 4 
 		gl.pixelStorei(gl.UNPACK_ALIGNMENT, 1);
 
-		// Create vertex buffer
+		// Create buffers
 		{
-			var vertices = [
+			this.gl_vertex_buffer = gl_create_vertex_buffer(gl, new Float32Array([
 				-1.0,+1.0,0.0,
 				-1.0,-1.0,0.0,
 				+1.0,-1.0,0.0,
 				+1.0,+1.0,0.0, 
-			];
+			]));
 
-			this.gl_vertex_buffer = gl.createBuffer();
-			gl.bindBuffer(gl.ARRAY_BUFFER, this.gl_vertex_buffer);
-			gl.bufferData(gl.ARRAY_BUFFER, new Float32Array(vertices), gl.STATIC_DRAW);
-			gl.bindBuffer(gl.ARRAY_BUFFER, null);
-		}
-
-		// Create vertex buffer
-		{
-			var indices = [0, 1, 2, 2, 0, 3];
-
-			this.gl_index_buffer = gl.createBuffer();
-			gl.bindBuffer(gl.ELEMENT_ARRAY_BUFFER, this.gl_index_buffer);
-			gl.bufferData(gl.ELEMENT_ARRAY_BUFFER, new Uint16Array(indices), gl.STATIC_DRAW);
-			gl.bindBuffer(gl.ELEMENT_ARRAY_BUFFER, null);
+			this.gl_index_buffer = gl_create_index_buffer(gl, new Uint16Array([0, 1, 2, 2, 0, 3]));
 		}
 
 		// Create simple shader program for NaN display mode
@@ -3291,51 +3503,6 @@ ${shader_display_assignment} = ${shader_display_operation};`;
 
 	compile_screen_shader_program(frag_code)
 	{
-		var gl = this.gl;
-
-		//console.log(frag_code);
-
-		const shader_std = `#version 300 es
-
-precision highp float;
-precision highp int;
-precision highp sampler2D;
-precision highp usampler2D;
-
-struct FloatInfos
-{
-	bool is_denormal;
-	bool is_infinity;
-	bool is_nan;
-};
-
-FloatInfos decodeFloat(float x)
-{
-	const uint mantissa_bit_count = 23u;
-	const uint exp_bit_count = 8u;
-
-	uint raw = floatBitsToUint(x);
-
-	uint sign_bit      = (raw >> (mantissa_bit_count + exp_bit_count)) & 0x1u;
-	uint mantissa_bits = (raw >> 0u)                     & ((0x1u << mantissa_bit_count) - 1u);
-	uint exp_bits      = (raw >> mantissa_bit_count)     & ((0x1u << exp_bit_count) - 1u);
-
-	bool is_max_exp = exp_bits == ((0x1u << exp_bit_count) - 1u);
-
-	FloatInfos infos;
-	infos.is_denormal = exp_bits == 0u;
-	infos.is_infinity = is_max_exp && mantissa_bits == 0u;
-	infos.is_nan = is_max_exp && mantissa_bits != 0u;
-
-	return infos;
-}
-
-#define isnan(x) decodeFloat(x).is_nan
-#define isinf(x) decodeFloat(x).is_infinity
-#define isdenorm(x) decodeFloat(x).is_denormal
-
-`;
-
 		var vert_code = `
 in vec3 coordinates;
 out highp vec2 uv;
@@ -3344,53 +3511,15 @@ void main(void) {
 	uv = coordinates.xy * 0.5 + 0.5;
 }`;
 
-		var vert_shader = gl.createShader(gl.VERTEX_SHADER);
-		{
-			gl.shaderSource(vert_shader, shader_std + vert_code);
-			gl.compileShader(vert_shader);
-
-			var compilation_status = gl.getShaderParameter(vert_shader, gl.COMPILE_STATUS);
-			if (!compilation_status)
-			{
-				var compilation_log = gl.getShaderInfoLog(vert_shader);
-				console.error('Fragment shader compiler log: ' + compilation_log);
-				return null;
-			}
-		}
-
-		var frag_shader = gl.createShader(gl.FRAGMENT_SHADER);
+		var shader_program = gl_create_shader_program(this.gl, vert_code, frag_code);
+		if (typeof shader_program === 'string')
 		{
-			gl.shaderSource(frag_shader, shader_std + frag_code); 
-			gl.compileShader(frag_shader);
-
-			var compilation_status = gl.getShaderParameter(frag_shader, gl.COMPILE_STATUS);
-			if (!compilation_status)
-			{
-				var compilation_log = gl.getShaderInfoLog(frag_shader);
-				console.log(frag_code);
-				console.error('Fragment shader compiler log: ' + compilation_log);
-				document.getElementById('texture_visualization_code_log').value = 'Compilation failed:\n' + compilation_log;
-				return null;
-			}
+			document.getElementById('texture_visualization_code_log').value = 'Compilation failed:\n' + shader_program;
+			return null;
 		}
-
-		var shader_program = gl.createProgram();
+		else
 		{
-			gl.attachShader(shader_program, vert_shader);
-			gl.attachShader(shader_program, frag_shader);
-			gl.linkProgram(shader_program);
-
-			var link_status = gl.getProgramParameter(shader_program, gl.LINK_STATUS);
-			if (!link_status)
-			{
-				var link_log = gl.getProgramInfoLog(shader_program);
-				console.error('Program compiler log: ' + link_log);
-				document.getElementById('texture_visualization_code_log').value = 'Link failed:\n' + link_log;
-			}
-			else
-			{
-				document.getElementById('texture_visualization_code_log').value = 'Compilation succeeded';
-			}
+			document.getElementById('texture_visualization_code_log').value = 'Compilation succeeded';
 		}
 		return shader_program;
 	}
@@ -4143,6 +4272,497 @@ function display_pass_parameters(pass_id)
 }
 
 
+// -------------------------------------------------------------------- GENERIC 3D STRUCTURE VIEW
+
+class Generic3DStructureView extends ResourceView
+{
+	constructor(subresource_version_info, resource_desc)
+	{
+		super(subresource_version_info, resource_desc);
+		this.is_ready = false;
+		this.release();
+		this.camera_pos = [200.0, 200.0, 200.0];
+		this.camera_fov = 90.0;
+		this.camera_longitude = 0.0;
+		this.camera_latitude = 0.0;
+		this.camera_look_at([0.0, 0.0, 0.0]);
+		this.camera_movement_u = 0;
+		this.camera_movement_v = 0;
+		this.camera_movement_z = 0;
+		this.camera_movement_speed = 1.0;
+		this.prev_time = 0;
+	}
+
+	setup_html(parent_dom)
+	{
+		var resource_info_htmls = '';
+		{
+			resource_info_htmls += `
+				<table width="100%" class="pretty_table resource_desc">`;
+			
+			for (var key in this.resource_desc)
+			{
+			    if (this.resource_desc.hasOwnProperty(key)){
+					resource_info_htmls += `
+						<tr>
+							<td>${key}</td>
+							<td>${this.resource_desc[key]}</td>
+						</tr>`;
+			    }
+			}
+
+			resource_info_htmls += `
+				</table>`;
+		}
+
+		var title = prettify_subresource_unique_name(this.subresource_version_info, this.resource_desc);
+		if (this.subresource_version_info['draw'] >= 0)
+		{
+			title += ` draw:${this.subresource_version_info['draw']}`;
+		}
+
+		var html = `
+			<div class="main_div">
+				<div class="selection_list_title">Ray-tracing Acceleration Structure visualization: ${title}</div>
+				<div id="canvas_outter">
+					<canvas
+						id="raytracing_visualization_canvas"
+						width="100"
+						height="100"
+						style="width: 100px; height: 100px;"></canvas>
+				</div>
+				<table width="100%">
+					<tr>
+						<td width="50%">
+							<div class="selection_list_title">Buffer descriptor</div>
+							${resource_info_htmls}
+						</td>
+					</tr>
+				</table>
+			</div>`;
+
+		parent_dom.innerHTML = html;
+
+		// Init WebGL 2.0
+		this.init_gl();
+
+		// Setup mouse motion callback
+		{
+			var generic_3d_view = this;
+
+			document.getElementById('canvas_outter').onclick = function(event) { generic_3d_view.onclick_canvas_outter(); };
+			document.getElementById('canvas_outter').onmouseleave = function(event) { generic_3d_view.set_select_viewport(false); }
+		}
+
+		this.init();
+	}
+
+	init_gl()
+	{
+		var gl_ctx_attributes = {
+			antialias: false,
+  			depth: false
+		};
+
+		// Init WebGL 2.0
+		this.canvas = document.getElementById('raytracing_visualization_canvas');
+
+		var gl_ctx_attributes = {
+			alpha: true,
+  			// premultipliedAlpha: true, // TODO
+			antialias: false,
+  			depth: true,
+  			stencil: false,
+  			powerPreference: "low-power",
+  			preserveDrawingBuffer: true,
+  			xrCompatible: false,
+		};
+
+		var gl = this.canvas.getContext('webgl2', gl_ctx_attributes);
+		this.gl = gl;
+
+		// Init WebGL extensions
+		{
+			var available_extensions = this.gl.getSupportedExtensions();
+			var required_extensions = ['EXT_texture_norm16'];
+
+			this.gl_ext = {};
+
+			var gl = this.gl;
+			var gl_ext = this.gl_ext;
+			required_extensions.forEach(function(ext_name)
+			{
+				gl_ext[ext_name] = gl.getExtension(ext_name);
+			});
+		}
+
+		// Create unit cube buffers exactly like GUnitCubeVertexBuffer and GUnitCubeIndexBuffer
+		{
+			var unit_cube_vertices = new Array(8 * 3);
+			for (var z = 0; z < 2; z++)
+			{
+				for (var y = 0; y < 2; y++)
+				{
+					for (var x = 0; x < 2; x++)
+					{
+						var dest = x * 4 + y * 2 + z;
+						unit_cube_vertices[dest * 3 + 0] = x ? -100.0 : 100.0;
+						unit_cube_vertices[dest * 3 + 1] = y ? -100.0 : 100.0;
+						unit_cube_vertices[dest * 3 + 2] = z ? -100.0 : 100.0;
+					}
+				}
+			}
+
+			var unit_cube_indices = [
+				0, 2, 3,
+				0, 3, 1,
+				4, 5, 7,
+				4, 7, 6,
+				0, 1, 5,
+				0, 5, 4,
+				2, 6, 7,
+				2, 7, 3,
+				0, 4, 6,
+				0, 6, 2,
+				1, 3, 7,
+				1, 7, 5,
+			];
+
+			this.gl_unit_cube = {}
+			this.gl_unit_cube.vertex_buffer = gl_create_vertex_buffer(gl, new Float32Array(unit_cube_vertices));
+			this.gl_unit_cube.index_buffer = gl_create_index_buffer(gl, new Uint16Array(unit_cube_indices));
+			this.gl_unit_cube.index_count = unit_cube_indices.length;
+		}
+
+		// Create simple shader program
+		{
+			var vert_code = `
+uniform mat4 local_to_world;
+uniform mat4 world_to_view;
+uniform mat4 view_to_clip;
+
+in vec3 vertex_coordinates;
+out vec3 interpolator_world;
+
+void main(void) {
+	//vec4 world_pos = local_to_world * vec4(vertex_coordinates.xy * clamp(vertex_coordinates.z, 0., 1.), vertex_coordinates.z, 1.0);
+	vec4 world_pos = local_to_world * vec4(vertex_coordinates.xyz, 1.0);
+	vec4 view_pos = world_to_view * world_pos;
+	vec4 clip_pos = view_to_clip * view_pos;
+
+	gl_Position = clip_pos;
+	interpolator_world = world_pos.xyz;
+}`;
+
+			var frag_code = `
+in vec3 interpolator_world;
+out vec4 display;
+
+void main(void) {
+	vec3 normal = normalize(cross(dFdx(interpolator_world), dFdy(interpolator_world)));
+
+	display = vec4(normal * 0.5 + 0.5, 1.0);
+}`;
+
+			this.gl_shader_program_simple = gl_create_shader_program(gl, vert_code, frag_code);
+		}
+
+		// Set unpacking alignement to 1 to allow uploading texture size not multple of 4 
+		gl.pixelStorei(gl.UNPACK_ALIGNMENT, 1);
+	}
+
+	release_gl()
+	{
+		this.canvas = null;
+		this.gl = null;
+		this.gl_ext = null;
+		this.gl_unit_cube_vertex_buffer = null;
+		this.gl_unit_cube_index_buffer = null;
+		this.gl_shader_program_simple = null;
+	}
+
+	draw_canvas()
+	{
+		console.error('unimplemented');
+	}
+
+	init()
+	{
+		console.error('unimplemented');
+	}
+
+	release()
+	{
+		this.viewport_selected = false;
+		this.release_gl();
+		this.set_select_viewport(false);
+	}
+
+	resize(ctx)
+	{
+		var dpi = window.devicePixelRatio || 1;
+		var viewport_width = ctx.width - 2 * k_style_padding_size - 2;
+		var viewport_height = viewport_width * 9.0 / 16.0;
+
+		var canvas_outter = document.getElementById('canvas_outter');
+
+		canvas_outter.style.width = `${viewport_width}px`;
+		canvas_outter.style.height = `${viewport_height}px`;
+
+		this.canvas.style.width = `${viewport_width}px`;
+		this.canvas.style.height = `${viewport_height}px`;
+		this.canvas.width = Math.round(viewport_width * dpi);
+		this.canvas.height = Math.round(viewport_height * dpi);
+	}
+
+	onclick_canvas_outter()
+	{
+		if (!this.is_ready)
+		{
+			return;
+		}
+
+		this.set_select_viewport(!this.viewport_selected);
+	}
+
+	set_select_viewport(selected)
+	{
+		if (this.viewport_selected === selected)
+		{
+			return;
+		}
+
+		this.viewport_selected = selected;
+
+		this.canvas.oncontextmenu = function(event) { return false; };
+
+		var canvas_outter = document.getElementById('canvas_outter');
+
+		if (this.viewport_selected)
+		{
+			var generic_3d_view = this;
+			canvas_outter.classList.add('selected');
+			//this.canvas.style.cursor = 'none';
+			this.canvas.onmousemove = function(event) { generic_3d_view.canvas_onmousemove(event); };
+			this.canvas.onwheel = function(event) { return generic_3d_view.canvas_onmousewheel(event); };
+
+			this.keydown_event = function(event) { generic_3d_view.canvas_onkeydown(event); };
+			this.keyup_event = function(event) { generic_3d_view.canvas_onkeyup(event); };
+			document.addEventListener(
+				"keydown",
+				this.keydown_event,
+				false);
+			document.addEventListener(
+				"keyup",
+				this.keyup_event,
+				false);
+
+			this.shedule_next_tick(performance.now());
+		}
+		else
+		{
+			canvas_outter.classList.remove('selected');
+			this.canvas.style.cursor = 'default';
+			this.canvas.onmousemove = null;
+			this.canvas.onwheel = null;
+
+			document.removeEventListener(
+				"keydown",
+				this.keydown_event);
+			document.removeEventListener(
+				"keyup",
+				this.keyup_event);
+
+			this.camera_movement_u = 0;
+			this.camera_movement_v = 0;
+			this.camera_movement_z = 0;
+		}
+	}
+
+	canvas_onmousemove(event)
+	{
+		var dx = event.movementX;
+		var dy = event.movementY;
+
+		var speed = (Math.PI / 180.0);
+
+		this.camera_longitude -= dx * speed;
+		this.camera_latitude = Math.min(Math.max(this.camera_latitude - dy * speed, -0.4 * Math.PI), 0.4 * Math.PI);
+	}
+
+	canvas_onmousewheel(event)
+	{
+		var zooming_in = event.deltaY < 0.0;
+
+		if (zooming_in)
+		{
+			this.camera_movement_speed *= 2.0;
+		}
+		else
+		{
+			this.camera_movement_speed *= 0.5;
+		}
+
+		return false;
+	}
+
+	canvas_onkeydown(event)
+	{
+		if (event.key == 'w')
+		{
+			this.camera_movement_z = 1.0;
+		}
+		else if (event.key == 's')
+		{
+			this.camera_movement_z = -1.0;
+		}
+		else if (event.key == 'q')
+		{
+			this.camera_movement_v = 1.0;
+		}
+		else if (event.key == 'e')
+		{
+			this.camera_movement_v = -1.0;
+		}
+		else if (event.key == 'a')
+		{
+			this.camera_movement_u = -1.0;
+		}
+		else if (event.key == 'd')
+		{
+			this.camera_movement_u = 1.0;
+		}
+	}
+
+	canvas_onkeyup(event)
+	{
+		if (event.key == 'w' && this.camera_movement_z > 0.0)
+		{
+			this.camera_movement_z = 0.0;
+		}
+		else if (event.key == 's' && this.camera_movement_z < 0.0)
+		{
+			this.camera_movement_z = 0.0;
+		}
+		else if (event.key == 'q' && this.camera_movement_v > 0.0)
+		{
+			this.camera_movement_v = 0.0;
+		}
+		else if (event.key == 'e' && this.camera_movement_v < 0.0)
+		{
+			this.camera_movement_v = 0.0;
+		}
+		else if (event.key == 'a' && this.camera_movement_u < 0.0)
+		{
+			this.camera_movement_u = 0.0;
+		}
+		else if (event.key == 'd' && this.camera_movement_u > 0.0)
+		{
+			this.camera_movement_u = 0.0;
+		}
+	}
+
+	camera_look_at(pos)
+	{
+		var camera_dir_x = pos[0] - this.camera_pos[0];
+		var camera_dir_y = pos[1] - this.camera_pos[1];
+		var camera_dir_z = pos[2] - this.camera_pos[2];
+		var camera_dir_len = Math.sqrt(camera_dir_x * camera_dir_x + camera_dir_y * camera_dir_y + camera_dir_z * camera_dir_z);
+
+		this.camera_longitude = Math.atan2(camera_dir_y / camera_dir_len, camera_dir_x / camera_dir_len);
+		this.camera_latitude = Math.asin(camera_dir_z / camera_dir_len);
+	}
+
+	get_view_matrix()
+	{
+		var camera_dir = [
+			Math.cos(this.camera_longitude) * Math.cos(this.camera_latitude),
+			Math.sin(this.camera_longitude) * Math.cos(this.camera_latitude),
+			Math.sin(this.camera_latitude),
+		];
+
+		return create_view_matrix(this.camera_pos, camera_dir);
+	}
+
+	get_proj_matrix()
+	{
+		if (false)
+		{
+			var width = 4.0;
+			var depth = 100.0;
+
+			var matrix = create_null_matrix();
+			matrix[0 + 0 * 4] = 1.0 / width;
+			matrix[1 + 1 * 4] = matrix[0 + 0 * 4] * this.canvas.width / this.canvas.height;
+			matrix[2 + 2 * 4] = - 1.0 / depth;
+			matrix[2 + 3 * 4] = 1.0;
+			matrix[3 + 3 * 4] = 1.0;
+			return matrix;
+		}
+
+		return create_projection_matrix_reverse_z_persepective(
+			this.camera_fov * 0.5, this.canvas.width / this.canvas.height, /* clipping_plane = */ 10.0);
+	}
+
+	shedule_next_tick(current_time)
+	{
+		var generic_3d_view = this;
+		this.prev_time = current_time;
+		requestAnimationFrame((new_time) => {
+    		generic_3d_view.tick(new_time);
+  		});
+	}
+
+	tick(current_time)
+	{
+		if (!this.viewport_selected)
+		{
+			return;
+		}
+
+		var delta_seconds = current_time - this.prev_time;
+
+		{
+			var delta_movements = 0.01 * delta_seconds * this.camera_movement_speed;
+
+			var camera_dir = [
+				Math.cos(this.camera_longitude) * Math.cos(this.camera_latitude),
+				Math.sin(this.camera_longitude) * Math.cos(this.camera_latitude),
+				Math.sin(this.camera_latitude),
+			];
+
+			var camera_dir_u = [
+				Math.sin(this.camera_longitude),
+				-Math.cos(this.camera_longitude),
+				0.0,
+			];
+
+			var camera_dir_v = [
+				Math.cos(this.camera_longitude) * Math.sin(-this.camera_latitude),
+				Math.sin(this.camera_longitude) * Math.sin(-this.camera_latitude),
+				Math.cos(this.camera_latitude),
+			];
+
+			this.camera_pos[0] += delta_movements * (
+				this.camera_movement_z * camera_dir[0] +
+				this.camera_movement_u * camera_dir_u[0] +
+				this.camera_movement_v * camera_dir_v[0]);
+			this.camera_pos[1] += delta_movements * (
+				this.camera_movement_z * camera_dir[1] +
+				this.camera_movement_u * camera_dir_u[1] +
+				this.camera_movement_v * camera_dir_v[1]);
+			this.camera_pos[2] += delta_movements * (
+				this.camera_movement_z * camera_dir[2] +
+				this.camera_movement_u * camera_dir_u[2] +
+				this.camera_movement_v * camera_dir_v[2]);
+		}
+
+		this.draw_canvas();
+		this.shedule_next_tick(current_time);
+	}
+} // class Generic3DStructureView
+
+
 // -------------------------------------------------------------------- UI
 
 function update_href_selection(parent)
