Implements RaytracingAccelerationStructureView capable of decoding file format of D3D12RayTracingSceneDebugUpdate()

Only thing left is to call D3D12RayTracingSceneDebugUpdate() in DumpGPU.cpp's AddDumpBufferPass() to write out the file.

diff --git a/Engine/Extras/GPUDumpViewer/GPUDumpViewer.html b/Engine/Extras/GPUDumpViewer/GPUDumpViewer.html
--- a/Engine/Extras/GPUDumpViewer/GPUDumpViewer.html
+++ b/Engine/Extras/GPUDumpViewer/GPUDumpViewer.html
@@ -1139,7 +1139,14 @@ function display_resource_internal(subresource_version_info)
 
 	if (resource_desc['Desc'] == 'FRDGBufferDesc')
 	{
-		g_view.set_resource_view(new BufferView(subresource_version_info, resource_desc));
+		if (resource_desc['Usage'].includes('AccelerationStructure'))
+		{
+			g_view.set_resource_view(new RaytracingAccelerationStructureView(subresource_version_info, resource_desc));
+		}
+		else
+		{
+			g_view.set_resource_view(new BufferView(subresource_version_info, resource_desc));
+		}
 	}
 	else if (resource_desc['Desc'] == 'FRDGTextureDesc' && (resource_desc['Type'] == 'Texture2D' || resource_desc['Type'] == 'Texture2DArray'))
 	{
@@ -4763,6 +4770,434 @@ void main(void) {
 } // class Generic3DStructureView
 
 
+// -------------------------------------------------------------------- HARDWARE ACCELERATION STRUCTURE
+
+class RaytracingAccelerationStructureView extends Generic3DStructureView
+{
+	constructor(subresource_version_info, resource_desc)
+	{
+		super(subresource_version_info, resource_desc);
+		this.draw_calls = [];
+	}
+
+	init()
+	{
+		var raytracing_view = this;
+		var subresource_version_name = get_subresource_unique_version_name(this.subresource_version_info);
+
+		load_resource_binary_file(`Resources/${subresource_version_name}.bin`, function(raw_raytracing_data)
+		{
+			if (raw_raytracing_data)
+			{
+				raytracing_view.parse_raw_raytracing_data(raw_raytracing_data);
+				raytracing_view.is_ready = true;
+			}
+			else
+			{
+				raytracing_view.error(`Error: Resources/${subresource_version_name}.bin couldn't be loaded.`);
+			}
+		});
+
+		document.title = `${g_view.pass_data['EventName']} - ${this.resource_desc['Name']}`;
+	}
+
+	parse_raw_raytracing_data(raw_raytracing_data)
+	{
+		var is_little_endian = true;
+
+		var data_view = new DataView(raw_raytracing_data);
+		var data_byte_offset = 0;
+
+		function set_byte_offset(byte_offset)
+		{
+			data_byte_offset = byte_offset;
+		}
+
+		function skip_bytes(byte_offset)
+		{
+			data_byte_offset += byte_offset;
+		}
+
+		function read_uint8()
+		{
+			var r = data_view.getUint8(data_byte_offset);
+			data_byte_offset += 1;
+			return r;
+		}
+
+		function read_uint32()
+		{
+			var r = data_view.getUint32(data_byte_offset, is_little_endian);
+			data_byte_offset += 4;
+			return r;
+		}
+
+		function read_uint64()
+		{
+			var l = BigInt(data_view.getUint32(data_byte_offset, is_little_endian));
+			var r = BigInt(data_view.getUint32(data_byte_offset + 4, is_little_endian));
+			data_byte_offset += 8;
+			return is_little_endian
+    			? l + BigInt(2 ** 32) * r
+    			: BigInt(2 ** 32) * l + r;
+		}
+
+		function read_uint32_array(size)
+		{
+			var r = new Array(size);
+			for (var i = 0; i < size; i++)
+			{
+				r[i] = data_view.getUint32(data_byte_offset, is_little_endian);
+				data_byte_offset += 4;
+			}
+			return r;
+		}
+
+		function read_float32_array(size)
+		{
+			var r = new Array(size);
+			for (var i = 0; i < size; i++)
+			{
+				r[i] = data_view.getFloat32(data_byte_offset, is_little_endian);
+				data_byte_offset += 4;
+			}
+			return r;
+		}
+
+		function assert_magic(magic)
+		{
+			var m = read_uint64();
+			if (magic !== m)
+			{
+				throw new Error(`Magic assertion ${magic} failed`);
+			}
+		}
+
+		const MaxNumLayers = 64;
+
+		// UE::HWRTScene::FSceneHeader
+		assert_magic(0x72ffa3376f48683bn);
+		assert_magic(1n);
+		var Offsets_Instances = read_uint32();
+		var Offsets_Geometries = read_uint32();
+		var Offsets_Buffers = read_uint32();
+		var Offsets_Strings = read_uint32();
+		var NumLayers = read_uint32();
+		var PerLayerNumInstances = read_uint32_array(MaxNumLayers);
+		var NumInstances = read_uint32();
+		var NumGeometries = read_uint32();
+		var NumBuffers = read_uint32();
+		var NumStrings = read_uint32();
+		skip_bytes(4 * 2);
+
+		set_byte_offset(Offsets_Instances);
+		var Instances = [];
+		for (var i = 0; i < NumInstances; i++)
+		{
+			var NewInstance = {};
+			Instances.push(NewInstance);
+
+			set_byte_offset(Offsets_Instances + i * (3 * 4 * 4 + 4 * 2 + 8));
+
+			// UE::HWRTScene::FInstanceDesc
+			NewInstance.Transform = read_float32_array(3 * 4);
+			var A = read_uint32();
+			var B = read_uint32();
+			NewInstance.InstanceID   = (A >>  0) & 0xFFFFFF;
+			NewInstance.InstanceMask = (A >> 24) & 0xFF;
+			NewInstance.InstanceContributionToHitGroupIndex = (B >>  0) & 0xFFFFFF;
+			NewInstance.Flags        = (B >> 24) & 0xFF;
+			NewInstance.AccelerationStructure = read_uint64();
+		}
+
+		// UE::HWRTScene::FGeometryHeader
+		set_byte_offset(Offsets_Geometries);
+		var Geometries = [];
+		for (var GeometryIndex = 0; GeometryIndex < NumGeometries; GeometryIndex++)
+		{
+			var Geometry = {};
+			Geometries.push(Geometry);
+
+			assert_magic(0x47226e42ad539683n);
+			Geometry.IndexBuffer = read_uint32();
+			Geometry.NumSegments = read_uint32();
+			skip_bytes(4 + 4 + 8);
+
+			// UE::HWRTScene::FSegmentHeader
+			Geometry.Segments = [];
+			for (var SegmentIndex = 0; SegmentIndex < Geometry.NumSegments; SegmentIndex++)
+			{
+				var Segment = {};
+				Geometry.Segments.push(Segment);
+
+				Segment.VertexBuffer = read_uint32();
+				Segment.VertexType = read_uint32();
+				Segment.VertexBufferOffset = read_uint32();
+				Segment.VertexBufferStride = read_uint32();
+				Segment.MaxVertices = read_uint32();
+				Segment.FirstPrimitive = read_uint32();
+				Segment.NumPrimitives = read_uint32();
+
+				Segment.bForceOpaque = read_uint8();
+				Segment.bAllowDuplicateAnyHitShaderInvocation = read_uint8();
+				Segment.bEnabled = read_uint8();
+				skip_bytes(1);
+			}
+		}
+
+		// UE::HWRTScene::FBufferHeader
+		set_byte_offset(Offsets_Buffers);
+		var RawBuffers = [];
+		for (var BufferIndex = 0; BufferIndex < NumBuffers; BufferIndex++)
+		{
+			var RawBuffer = {};
+			RawBuffers.push(RawBuffer);
+
+			assert_magic(0x7330d54d0195a6den);
+			RawBuffer.SizeInBytes = read_uint32();
+			RawBuffer.StrideInBytes = read_uint32();
+			RawBuffer.ByteOffset = data_byte_offset;
+			skip_bytes(RawBuffer.SizeInBytes);
+		}
+
+		// Translate UE::HWRTScene to webgl draw calls
+		var webgl_buffers = new Array(RawBuffers.length);
+		var min_position = [0, 0, 0];
+		var max_position = [0, 0, 0];
+		var gl = this.gl;
+		for (var instance_index = 0; instance_index < NumInstances; instance_index++)
+		{
+			var Instance = Instances[instance_index];
+			var Geometry = Geometries[Instance.AccelerationStructure];
+			var IndexBuffer = RawBuffers[Geometry.IndexBuffer];
+
+			// only support indexed geometry right now
+			if (IndexBuffer.SizeInBytes == 0)
+			{
+				continue;
+			}
+
+			// Upload index buffer to GPU if not already.
+			if (!webgl_buffers[Geometry.IndexBuffer])
+			{
+				var indices = null;
+				if (IndexBuffer.StrideInBytes == 4)
+				{
+					indices = new Uint32Array(raw_raytracing_data, IndexBuffer.ByteOffset, IndexBuffer.SizeInBytes / 4);
+				}
+				else if (IndexBuffer.StrideInBytes == 2)
+				{
+					indices = new Uint16Array(raw_raytracing_data, IndexBuffer.ByteOffset, IndexBuffer.SizeInBytes / 2);
+				}
+				else if (IndexBuffer.StrideInBytes == 1)
+				{
+					indices = new Uint8Array(raw_raytracing_data, IndexBuffer.ByteOffset, IndexBuffer.SizeInBytes / 1);
+				}
+				else
+				{
+					throw Error(`Unknown IndexBuffer.StrideInBytes=${IndexBuffer.StrideInBytes}`);
+				}
+				console.assert(indices.byteOffset == IndexBuffer.ByteOffset);
+				console.assert(indices.byteLength == IndexBuffer.SizeInBytes);
+
+				webgl_buffers[Geometry.IndexBuffer] = gl_create_index_buffer(gl, indices);
+			}
+
+			// setup local to world
+			var local_to_world = create_identity_matrix();
+			{
+				local_to_world[0 + 0 * 4] = Instance.Transform[0 * 4 + 0];
+				local_to_world[1 + 0 * 4] = Instance.Transform[1 * 4 + 0];
+				local_to_world[2 + 0 * 4] = Instance.Transform[2 * 4 + 0];
+				local_to_world[0 + 1 * 4] = Instance.Transform[0 * 4 + 1];
+				local_to_world[1 + 1 * 4] = Instance.Transform[1 * 4 + 1];
+				local_to_world[2 + 1 * 4] = Instance.Transform[2 * 4 + 1];
+				local_to_world[0 + 2 * 4] = Instance.Transform[0 * 4 + 2];
+				local_to_world[1 + 2 * 4] = Instance.Transform[1 * 4 + 2];
+				local_to_world[2 + 2 * 4] = Instance.Transform[2 * 4 + 2];
+				local_to_world[0 + 3 * 4] = Instance.Transform[0 * 4 + 3];
+				local_to_world[1 + 3 * 4] = Instance.Transform[1 * 4 + 3];
+				local_to_world[2 + 3 * 4] = Instance.Transform[2 * 4 + 3];
+			}
+
+			for (var SegmentIndex = 0; SegmentIndex < Geometry.NumSegments; SegmentIndex++)
+			{
+				var Segment = Geometry.Segments[SegmentIndex];
+				var VertexBuffer = RawBuffers[Segment.VertexBuffer];
+
+				console.assert(Segment.VertexType == 3); // VET_Float3
+
+				// Upload vertex buffer to GPU if not already.
+				if (!webgl_buffers[Segment.VertexBuffer])
+				{
+					var vertices = new Float32Array(raw_raytracing_data, VertexBuffer.ByteOffset, VertexBuffer.SizeInBytes / 4);
+					console.assert(vertices.byteOffset == VertexBuffer.ByteOffset);
+					console.assert(vertices.byteLength == VertexBuffer.SizeInBytes);
+					webgl_buffers[Segment.VertexBuffer] = gl_create_vertex_buffer(gl, vertices);
+				}
+
+				var draw_call = {};
+				this.draw_calls.push(draw_call);
+
+				draw_call.local_to_world = local_to_world;
+				if (IndexBuffer.StrideInBytes == 4)
+				{
+					draw_call.index_type = gl.UNSIGNED_INT;
+				}
+				else if (IndexBuffer.StrideInBytes == 2)
+				{
+					draw_call.index_type = gl.UNSIGNED_SHORT;
+				}
+				else if (IndexBuffer.StrideInBytes == 1)
+				{
+					draw_call.index_type = gl.UNSIGNED_BYTE;
+				}
+				else
+				{
+					throw Error(`Unknown IndexBuffer.StrideInBytes=${IndexBuffer.StrideInBytes}`);
+				}
+				draw_call.index_byte_offset = Segment.FirstPrimitive * 3 * IndexBuffer.StrideInBytes;
+				draw_call.index_count = Segment.NumPrimitives * 3;
+
+				// create vertex array.
+				{
+					console.assert(Segment.VertexBufferStride == 12);
+					console.assert((Segment.VertexBufferOffset % Segment.VertexBufferStride) == 0);
+
+					draw_call.vertex_array = gl.createVertexArray();
+					gl.bindVertexArray(draw_call.vertex_array);
+					gl.bindBuffer(gl.ARRAY_BUFFER, webgl_buffers[Segment.VertexBuffer]);
+					gl.bindBuffer(gl.ELEMENT_ARRAY_BUFFER, webgl_buffers[Geometry.IndexBuffer]);
+					gl.enableVertexAttribArray(0);
+					gl.vertexAttribPointer(0, 3, gl.FLOAT, false, Segment.VertexBufferStride, Segment.VertexBufferOffset);
+					gl.bindVertexArray(null);
+					console.assert(gl.getError() == gl.NO_ERROR); 
+				}
+
+				// Track the min and max instance position for seting up the camera over the entire scene.
+				if (instance_index == 0)
+				{
+					min_position[0] = Instance.Transform[0 * 4 + 3];
+					min_position[1] = Instance.Transform[1 * 4 + 3];
+					min_position[2] = Instance.Transform[2 * 4 + 3];
+
+					max_position[0] = Instance.Transform[0 * 4 + 3];
+					max_position[1] = Instance.Transform[1 * 4 + 3];
+					max_position[2] = Instance.Transform[2 * 4 + 3];
+				}
+				else
+				{
+					min_position[0] = Math.min(min_position[0], Instance.Transform[0 * 4 + 3]);
+					min_position[1] = Math.min(min_position[0], Instance.Transform[1 * 4 + 3]);
+					min_position[2] = Math.min(min_position[0], Instance.Transform[2 * 4 + 3]);
+
+					max_position[0] = Math.max(max_position[0], Instance.Transform[0 * 4 + 3]);
+					max_position[1] = Math.max(max_position[0], Instance.Transform[1 * 4 + 3]);
+					max_position[2] = Math.max(max_position[0], Instance.Transform[2 * 4 + 3]);
+				}
+			}
+		}
+
+		if (true)
+		{
+			this.camera_pos = max_position;
+			this.camera_look_at([(min_position[0] + max_position[0]) / 2, (min_position[1] + max_position[1]) / 2, (min_position[2] + max_position[2]) / 2]);
+		}
+
+		this.draw_canvas();
+	} // parse_raw_raytracing_data()
+
+	draw_canvas()
+	{
+		var gl = this.gl;
+
+		var world_to_view = this.get_view_matrix();
+		var view_to_clip = this.get_proj_matrix();
+
+		// view_matrix[0 + 0 * 4] = 0.5;
+		// view_matrix[0 + 1 * 4] = 0.5;
+		// view_matrix[1 + 1 * 4] = 0.5;
+		// view_matrix[2 + 2 * 4] = 0.5;
+
+		gl.viewport(0, 0, this.canvas.width, this.canvas.height);
+		gl.clearColor(0.0, 0.0, 0.0, 0.0);
+		gl.clear(gl.COLOR_BUFFER_BIT);
+		
+		gl.clearDepth(0.0);
+		gl.clear(gl.DEPTH_BUFFER_BIT);
+
+		gl.enable(gl.DEPTH_TEST);
+		gl.depthFunc(gl.GEQUAL);
+
+		gl.disable(gl.CULL_FACE);
+		gl.cullFace(gl.FRONT_AND_BACK);
+
+		{
+			var draw_cubes = false;
+
+			var shader_program = this.gl_shader_program_simple;
+			gl.useProgram(shader_program);
+
+			var coord = gl.getAttribLocation(shader_program, "vertex_coordinates");
+			console.assert(coord == 0);
+
+			if (draw_cubes)
+			{
+				gl.bindBuffer(gl.ARRAY_BUFFER, this.gl_unit_cube.vertex_buffer);
+				gl.bindBuffer(gl.ELEMENT_ARRAY_BUFFER, this.gl_unit_cube.index_buffer);
+				gl.vertexAttribPointer(coord, 3, gl.FLOAT, false, 0, 0); 
+				gl.enableVertexAttribArray(coord);
+			}
+
+			gl_set_uniform_mat4(gl, shader_program, 'local_to_world', create_identity_matrix());
+			gl_set_uniform_mat4(gl, shader_program, 'world_to_view', world_to_view);
+			gl_set_uniform_mat4(gl, shader_program, 'view_to_clip', view_to_clip);
+
+			for (var draw_call of this.draw_calls)
+			{
+				gl_set_uniform_mat4(gl, shader_program, 'local_to_world', draw_call.local_to_world);
+
+				if (!draw_cubes)
+				{
+					gl.bindVertexArray(draw_call.vertex_array);
+					gl.drawElementsInstanced(
+						gl.TRIANGLES,
+						/* index_count = */       draw_call.index_count,
+						/* index_type = */        draw_call.index_type,
+						/* index_byte_offset = */ draw_call.index_byte_offset,
+						/* instanceCount = */     1);
+					gl.bindVertexArray(null);
+					/*var gl_error = gl.getError();
+					if (gl_error == gl.NO_ERROR)
+					{
+						//throw Error();
+					}*/
+				}
+				else
+				{
+					gl.drawElementsInstanced(
+						gl.TRIANGLES,
+						/* index_count = */       this.gl_unit_cube.index_count,
+						/* index_type = */        gl.UNSIGNED_SHORT,
+						/* index_byte_offset = */ 0,
+						/* instanceCount = */     1);
+				}
+			}
+
+			gl.useProgram(null);
+			gl.bindBuffer(gl.ARRAY_BUFFER, null);
+			gl.bindBuffer(gl.ELEMENT_ARRAY_BUFFER, null);
+		}
+	}
+
+	release()
+	{
+		super.release();
+		this.draw_calls = []
+	}
+} // class RaytracingAccelerationStructureView
+
+
 // -------------------------------------------------------------------- UI
 
 function update_href_selection(parent)
