Implements copy to clipboard and scene visualization modes in ray tracing scene viewer

diff --git a/Engine/Extras/GPUDumpViewer/GPUDumpViewer.html b/Engine/Extras/GPUDumpViewer/GPUDumpViewer.html
--- a/Engine/Extras/GPUDumpViewer/GPUDumpViewer.html
+++ b/Engine/Extras/GPUDumpViewer/GPUDumpViewer.html
@@ -1682,7 +1682,20 @@ function display_log()
 
 
 
-// -------------------------------------------------------------------- WEBGL UTILS TEXTURE
+// -------------------------------------------------------------------- WEBGL UTILS
+
+function copy_canvas_to_clipboard(canvas, text)
+{
+	canvas.toBlob(function(image_blob) {
+		var text_blob = new Blob([text], { type: 'text/plain' });
+		var clipboard_data = {
+			[text_blob.type]: text_blob,
+			[image_blob.type]: image_blob,
+		};
+
+		navigator.clipboard.write([new ClipboardItem(clipboard_data)]);
+	}, 'image/png');
+}
 
 function gl_create_vertex_buffer(gl, vertices)
 {
@@ -1814,12 +1827,22 @@ function gl_create_shader_program(gl, vert_shader, frag_shader)
 	return shader_program;
 }
 
+function gl_set_uniform_uint(gl, program, name, value)
+{
+	var uniform_loc = gl.getUniformLocation(program, name);
+	if (uniform_loc)
+	{
+		gl.uniform1ui(uniform_loc, value);
+	}
+}
+
 function gl_set_uniform_mat4(gl, program, name, matrix)
 {
-	gl.uniformMatrix4fv(
-		gl.getUniformLocation(program, name),
-		false,
-		new Float32Array(matrix));
+	var uniform_loc = gl.getUniformLocation(program, name);
+	if (uniform_loc)
+	{
+		gl.uniformMatrix4fv(uniform_loc, false, new Float32Array(matrix));
+	}
 }
 
 
@@ -2789,44 +2812,35 @@ void main(void)
 			return;
 		}
 		
-		var texture_view = this;
-		this.canvas.toBlob(function(image_blob) {
-			var pass_data = g_view.pass_data;
+		var pass_data = g_view.pass_data;
 
-			var is_output_resource = is_pass_output_resource(pass_data, texture_view.subresource_version_info);
+		var is_output_resource = is_pass_output_resource(pass_data, this.subresource_version_info);
 
-			var text = '';
-			text += `Dump:\n\t${g_infos['Project']} - ${g_infos['Platform']} - ${g_infos['RHI']} - ${g_infos['BuildVersion']} - ${g_infos['DumpTime']}\n\t${get_current_navigation()}\n`
-			text += `Pass:\n\t${pass_data['EventName']}\n`;
+		var text = '';
+		text += `Dump:\n\t${g_infos['Project']} - ${g_infos['Platform']} - ${g_infos['RHI']} - ${g_infos['BuildVersion']} - ${g_infos['DumpTime']}\n\t${get_current_navigation()}\n`
+		text += `Pass:\n\t${pass_data['EventName']}\n`;
 
+		{
+			var name = prettify_subresource_unique_name(this.subresource_version_info, this.resource_desc);
+			text += `${is_output_resource ? 'Output' : 'Input'} resource:\n\t${name}\n`;
+			if (this.display_mode == 'Visualization')
 			{
-				var name = prettify_subresource_unique_name(texture_view.subresource_version_info, texture_view.resource_desc);
-				text += `${is_output_resource ? 'Output' : 'Input'} resource:\n\t${name}\n`;
-				if (texture_view.display_mode == 'Visualization')
-				{
-					text += `\tDisplay shader: \n${g_shader_code_dict[texture_view.shader_code_saving_key]}\n`;
-					text += `\tDisplay channels: ${texture_view.src_channels}\n`;
-					text += `\tDisplay src color: ${texture_view.src_color_space}\n`;
-				}
-				else
-				{
-					text += `\tSrc color: ${texture_view.display_mode}\n`;
-				}
+				text += `\tDisplay shader: \n${g_shader_code_dict[this.shader_code_saving_key]}\n`;
+				text += `\tDisplay channels: ${this.src_channels}\n`;
+				text += `\tDisplay src color: ${this.src_color_space}\n`;
 			}
-
-			if (texture_view.subresource_version_info['draw'] >= 0)
+			else
 			{
-				text += `Draw ${texture_view.subresource_version_info['draw']}:\n\t${g_view.pass_draws_data[texture_view.subresource_version_info['draw']]['DrawName']}\n`;
+				text += `\tSrc color: ${this.display_mode}\n`;
 			}
+		}
 
-			var text_blob = new Blob([text], { type: 'text/plain' });
-			var clipboard_data = {
-				[text_blob.type]: text_blob,
-				[image_blob.type]: image_blob,
-			};
+		if (this.subresource_version_info['draw'] >= 0)
+		{
+			text += `Draw ${this.subresource_version_info['draw']}:\n\t${g_view.pass_draws_data[this.subresource_version_info['draw']]['DrawName']}\n`;
+		}
 
-			navigator.clipboard.write([new ClipboardItem(clipboard_data)]);
-		}, 'image/png');
+		copy_canvas_to_clipboard(this.canvas, text);
 	}
 
 	get shader_code_saving_key()
@@ -4296,7 +4310,7 @@ class Generic3DStructureView extends ResourceView
 		this.camera_movement_u = 0;
 		this.camera_movement_v = 0;
 		this.camera_movement_z = 0;
-		this.camera_movement_speed = 1.0;
+		this.camera_movement_speed = 100.0; // 1 m/s
 		this.prev_time = 0;
 	}
 
@@ -4332,8 +4346,13 @@ class Generic3DStructureView extends ResourceView
 			<div class="main_div">
 				<div class="selection_list_title">Ray-tracing Acceleration Structure visualization: ${title}</div>
 				<div id="canvas_outter">
+					<div id="canvas_header">
+						<div class="button_list">
+							<input type="button" onclick="g_view.resource_view.copy_canvas_to_clipboard();" value="Copy to clipboard"/>
+						</div>
+					</div>
 					<canvas
-						id="raytracing_visualization_canvas"
+						id="generic_3d_visualization_canvas"
 						width="100"
 						height="100"
 						style="width: 100px; height: 100px;"></canvas>
@@ -4357,7 +4376,7 @@ class Generic3DStructureView extends ResourceView
 		{
 			var generic_3d_view = this;
 
-			document.getElementById('canvas_outter').onclick = function(event) { generic_3d_view.onclick_canvas_outter(); };
+			this.canvas.onclick = function(event) { generic_3d_view.onclick_canvas(); };
 			document.getElementById('canvas_outter').onmouseleave = function(event) { generic_3d_view.set_select_viewport(false); }
 		}
 
@@ -4372,7 +4391,7 @@ class Generic3DStructureView extends ResourceView
 		};
 
 		// Init WebGL 2.0
-		this.canvas = document.getElementById('raytracing_visualization_canvas');
+		this.canvas = document.getElementById('generic_3d_visualization_canvas');
 
 		var gl_ctx_attributes = {
 			alpha: true,
@@ -4452,7 +4471,6 @@ in vec3 vertex_coordinates;
 out vec3 interpolator_world;
 
 void main(void) {
-	//vec4 world_pos = local_to_world * vec4(vertex_coordinates.xy * clamp(vertex_coordinates.z, 0., 1.), vertex_coordinates.z, 1.0);
 	vec4 world_pos = local_to_world * vec4(vertex_coordinates.xyz, 1.0);
 	vec4 view_pos = world_to_view * world_pos;
 	vec4 clip_pos = view_to_clip * view_pos;
@@ -4513,16 +4531,13 @@ void main(void) {
 
 		var canvas_outter = document.getElementById('canvas_outter');
 
-		canvas_outter.style.width = `${viewport_width}px`;
-		canvas_outter.style.height = `${viewport_height}px`;
-
 		this.canvas.style.width = `${viewport_width}px`;
 		this.canvas.style.height = `${viewport_height}px`;
 		this.canvas.width = Math.round(viewport_width * dpi);
 		this.canvas.height = Math.round(viewport_height * dpi);
 	}
 
-	onclick_canvas_outter()
+	onclick_canvas()
 	{
 		if (!this.is_ready)
 		{
@@ -4669,6 +4684,29 @@ void main(void) {
 		}
 	}
 
+	copy_canvas_to_clipboard()
+	{
+		if (!this.is_ready)
+		{
+			return;
+		}
+		
+		var pass_data = g_view.pass_data;
+
+		var is_output_resource = is_pass_output_resource(pass_data, this.subresource_version_info);
+
+		var text = '';
+		text += `Dump:\n\t${g_infos['Project']} - ${g_infos['Platform']} - ${g_infos['RHI']} - ${g_infos['BuildVersion']} - ${g_infos['DumpTime']}\n\t${get_current_navigation()}\n`
+		text += `Pass:\n\t${pass_data['EventName']}\n`;
+
+		{
+			var name = prettify_subresource_unique_name(this.subresource_version_info, this.resource_desc);
+			text += `${is_output_resource ? 'Output' : 'Input'} resource:\n\t${name}\n`;
+		}
+
+		copy_canvas_to_clipboard(this.canvas, text);
+	}
+
 	camera_look_at(pos)
 	{
 		var camera_dir_x = pos[0] - this.camera_pos[0];
@@ -4778,10 +4816,95 @@ class RaytracingAccelerationStructureView extends Generic3DStructureView
 	{
 		super(subresource_version_info, resource_desc);
 		this.draw_calls = [];
+		this.scene_visualization_shader_programs = {};
+		this.scene_visualization = 'Instances';
 	}
 
 	init()
 	{
+		document.getElementById('canvas_header').innerHTML += `
+			<div class="button_list" id="change_scene_visualization"><!---
+				---><input type="button" onclick="g_view.resource_view.change_scene_visualization(this);" value="Instances"/><!--- 
+				---><input type="button" onclick="g_view.resource_view.change_scene_visualization(this);" value="Normals"/>
+			</div>`;
+
+		update_value_selection(document.getElementById('change_scene_visualization'), this.scene_visualization);
+
+		// Create visualization shader programs
+		this.scene_visualization_shader_programs = {};
+		{
+			var vert_code = `
+uniform mat4 local_to_world;
+uniform mat4 world_to_view;
+uniform mat4 view_to_clip;
+
+in vec3 vertex_coordinates;
+out vec3 interpolator_world;
+
+void main(void) {
+	vec4 world_pos = local_to_world * vec4(vertex_coordinates.xyz, 1.0);
+	vec4 view_pos = world_to_view * world_pos;
+	vec4 clip_pos = view_to_clip * view_pos;
+
+	gl_Position = clip_pos;
+	interpolator_world = world_pos.xyz;
+}`;
+
+			var frag_normals_code = `
+in vec3 interpolator_world;
+out vec4 display;
+
+void main(void) {
+	vec3 normal = normalize(cross(dFdx(interpolator_world), dFdy(interpolator_world)));
+
+	display = vec4(normal * 0.5 + 0.5, 1.0);
+}`;
+
+			var frag_instances_code = `
+uniform uint instance_index;
+
+in vec3 interpolator_world;
+out vec4 display;
+
+uint MurmurMix(uint Hash)
+{
+	Hash ^= Hash >> 16u;
+	Hash *= 0x85ebca6bu;
+	Hash ^= Hash >> 13u;
+	Hash *= 0xc2b2ae35u;
+	Hash ^= Hash >> 16u;
+	return Hash;
+}
+
+vec3 IntToColor(uint Index)
+{
+	uint Hash = MurmurMix(Index);
+
+	vec3 Color = vec3
+	(
+		(Hash >>  0u) & 255u,
+		(Hash >>  8u) & 255u,
+		(Hash >> 16u) & 255u
+	);
+
+	return Color * (1.0f / 255.0f);
+}
+
+void main(void) {
+	vec3 normal = normalize(cross(dFdx(interpolator_world), dFdy(interpolator_world)));
+
+	vec3 vert_color = IntToColor(1u + instance_index);
+	vec3 base_color = vert_color * pow(gl_FragCoord.w, 0.1);
+
+	display = vec4(mix(base_color, normal * 0.5 + 0.5, 1.0 / 3.0), 1.0);
+}`;
+
+			this.scene_visualization_shader_programs['Normals'] = gl_create_shader_program(
+				this.gl, vert_code, frag_normals_code);
+			this.scene_visualization_shader_programs['Instances'] = gl_create_shader_program(
+				this.gl, vert_code, frag_instances_code);
+		}
+
 		var raytracing_view = this;
 		var subresource_version_name = get_subresource_unique_version_name(this.subresource_version_info);
 
@@ -5040,6 +5163,7 @@ class RaytracingAccelerationStructureView extends Generic3DStructureView
 				this.draw_calls.push(draw_call);
 
 				draw_call.local_to_world = local_to_world;
+				draw_call.instance_index = instance_index;
 				if (IndexBuffer.StrideInBytes == 4)
 				{
 					draw_call.index_type = gl.UNSIGNED_INT;
@@ -5107,6 +5231,23 @@ class RaytracingAccelerationStructureView extends Generic3DStructureView
 		this.draw_canvas();
 	} // parse_raw_raytracing_data()
 
+	change_scene_visualization(new_scene_visualization)
+	{
+		if (!this.is_ready)
+		{
+			return;
+		}
+		
+		if (new_scene_visualization.value == this.scene_visualization)
+		{
+			return;
+		}
+
+		this.scene_visualization = new_scene_visualization.value;
+		update_value_selection(document.getElementById('change_scene_visualization'), this.scene_visualization);
+		this.draw_canvas();
+	}
+
 	draw_canvas()
 	{
 		var gl = this.gl;
@@ -5133,22 +5274,12 @@ class RaytracingAccelerationStructureView extends Generic3DStructureView
 		gl.cullFace(gl.FRONT_AND_BACK);
 
 		{
-			var draw_cubes = false;
-
-			var shader_program = this.gl_shader_program_simple;
+			var shader_program = this.scene_visualization_shader_programs[this.scene_visualization];
 			gl.useProgram(shader_program);
 
 			var coord = gl.getAttribLocation(shader_program, "vertex_coordinates");
 			console.assert(coord == 0);
 
-			if (draw_cubes)
-			{
-				gl.bindBuffer(gl.ARRAY_BUFFER, this.gl_unit_cube.vertex_buffer);
-				gl.bindBuffer(gl.ELEMENT_ARRAY_BUFFER, this.gl_unit_cube.index_buffer);
-				gl.vertexAttribPointer(coord, 3, gl.FLOAT, false, 0, 0); 
-				gl.enableVertexAttribArray(coord);
-			}
-
 			gl_set_uniform_mat4(gl, shader_program, 'local_to_world', create_identity_matrix());
 			gl_set_uniform_mat4(gl, shader_program, 'world_to_view', world_to_view);
 			gl_set_uniform_mat4(gl, shader_program, 'view_to_clip', view_to_clip);
@@ -5156,37 +5287,19 @@ class RaytracingAccelerationStructureView extends Generic3DStructureView
 			for (var draw_call of this.draw_calls)
 			{
 				gl_set_uniform_mat4(gl, shader_program, 'local_to_world', draw_call.local_to_world);
-
-				if (!draw_cubes)
-				{
-					gl.bindVertexArray(draw_call.vertex_array);
-					gl.drawElementsInstanced(
-						gl.TRIANGLES,
-						/* index_count = */       draw_call.index_count,
-						/* index_type = */        draw_call.index_type,
-						/* index_byte_offset = */ draw_call.index_byte_offset,
-						/* instanceCount = */     1);
-					gl.bindVertexArray(null);
-					/*var gl_error = gl.getError();
-					if (gl_error == gl.NO_ERROR)
-					{
-						//throw Error();
-					}*/
-				}
-				else
-				{
-					gl.drawElementsInstanced(
-						gl.TRIANGLES,
-						/* index_count = */       this.gl_unit_cube.index_count,
-						/* index_type = */        gl.UNSIGNED_SHORT,
-						/* index_byte_offset = */ 0,
-						/* instanceCount = */     1);
-				}
+				gl_set_uniform_uint(gl, shader_program, 'instance_index', draw_call.instance_index);
+
+				gl.bindVertexArray(draw_call.vertex_array);
+				gl.drawElementsInstanced(
+					gl.TRIANGLES,
+					/* index_count = */       draw_call.index_count,
+					/* index_type = */        draw_call.index_type,
+					/* index_byte_offset = */ draw_call.index_byte_offset,
+					/* instanceCount = */     1);
 			}
 
 			gl.useProgram(null);
-			gl.bindBuffer(gl.ARRAY_BUFFER, null);
-			gl.bindBuffer(gl.ELEMENT_ARRAY_BUFFER, null);
+			gl.bindVertexArray(null);
 		}
 	}
 
@@ -5194,6 +5307,7 @@ class RaytracingAccelerationStructureView extends Generic3DStructureView
 	{
 		super.release();
 		this.draw_calls = []
+		this.scene_visualization_shader_programs = {};
 	}
 } // class RaytracingAccelerationStructureView
 
