Fixes the GPU dump viewer frame incrementer on TSR's RejectShading pass

diff --git a/Engine/Extras/GPUDumpViewer/GPUDumpViewer.html b/Engine/Extras/GPUDumpViewer/GPUDumpViewer.html
--- a/Engine/Extras/GPUDumpViewer/GPUDumpViewer.html
+++ b/Engine/Extras/GPUDumpViewer/GPUDumpViewer.html
@@ -512,6 +512,12 @@ function get_frame_number(pass_data)
 
 function compare_draw_events(ref_draw_event, tested_draw_event)
 {
+	ref_draw_event = ref_draw_event.replace(/\d+\.d*/g, '');
+	ref_draw_event = ref_draw_event.replace(/\d+/g, '');
+
+	tested_draw_event = tested_draw_event.replace(/\d+\.d*/g, '');
+	tested_draw_event = tested_draw_event.replace(/\d+/g, '');
+
 	return ref_draw_event == tested_draw_event;
 }
 
@@ -601,6 +607,25 @@ function find_similar_resource_in_pass(current_frame_pass_data, current_frame_re
 		return new_pass_subresource_info;
 	}
 
+	// try to find with non matching array_slice
+	for (var new_pass_subresource_unique_name of new_pass_resource_list)
+	{
+		var new_pass_subresource_info = parse_subresource_unique_name(new_pass_subresource_unique_name);
+		if (new_pass_subresource_info['subresource'] !== current_frame_subresource_info['subresource'])
+		{
+			continue;
+		}
+
+		var new_pass_resource_desc = get_resource_desc(new_pass_subresource_info['resource']);
+		if (new_pass_resource_desc['Name'] !== current_frame_resource_desc['Name'])
+		{
+			continue;
+		}
+
+		return new_pass_subresource_info;
+	}
+
+
 	return null;
 }
 
@@ -837,6 +862,12 @@ class PassView extends IView
 		var prev_frame_pass_id = find_similar_pass_id_in_frame(this.pass_data, prev_frame_number);
 		var next_frame_pass_id = find_similar_pass_id_in_frame(this.pass_data, next_frame_number);
 
+		if (prev_frame_pass_id == null || next_frame_pass_id == null)
+		{
+			document.getElementById('pass_frame_list').style.display = 'none';
+			return;
+		}
+
 		var prev_frame_href = '';
 		var next_frame_href = '';
 		if (this.resource_view.resource_desc['Desc'] == 'FRDGParameterStruct')
@@ -855,7 +886,12 @@ class PassView extends IView
 
 			var prev_frame_href = '';
 			var next_frame_href = '';
-			if (is_output_resource)
+			if (prev_frame_subresource_info == null || next_frame_subresource_info == null)
+			{
+				document.getElementById('pass_frame_list').style.display = 'none';
+				return;
+			}
+			else if (is_output_resource)
 			{
 				prev_frame_href = `#display_output_resource(${prev_frame_pass_id},'${get_subresource_unique_name(prev_frame_subresource_info)}');`;
 				next_frame_href = `#display_output_resource(${next_frame_pass_id},'${get_subresource_unique_name(next_frame_subresource_info)}');`;
