Fix GPUDumpViewer so it will still work on Linux when it exists within a directory with quotes somewhere along the path

diff --git a/Engine/Extras/GPUDumpViewer/OpenGPUDumpViewer.sh b/Engine/Extras/GPUDumpViewer/OpenGPUDumpViewer.sh
--- a/Engine/Extras/GPUDumpViewer/OpenGPUDumpViewer.sh
+++ b/Engine/Extras/GPUDumpViewer/OpenGPUDumpViewer.sh
@@ -12,7 +12,7 @@ SCRIPT=$(readlink -f "$0")
 # Absolute path this script is in, thus [REDACTED_USER_HOME]/bin
 SCRIPTPATH=$(dirname "$SCRIPT")
 
-pushd $SCRIPTPATH
+pushd "$SCRIPTPATH"
 
 
 GetAllChildProcesses() {
