Adds supports for PF_R8G8_UINT in the GPUDumpViewer

diff --git a/Engine/Extras/GPUDumpViewer/GPUDumpViewer.html b/Engine/Extras/GPUDumpViewer/GPUDumpViewer.html
--- a/Engine/Extras/GPUDumpViewer/GPUDumpViewer.html
+++ b/Engine/Extras/GPUDumpViewer/GPUDumpViewer.html
@@ -3130,6 +3130,15 @@ void main(void)
 			raw_channel_count = 4;
 			raw_channel_format = 'uint';
 		}
+		else if (ue_pixel_format === 'PF_R8G8_UINT')
+		{
+			webgl_texture_internal_format = gl.RG8UI;
+			webgl_texture_src_format = gl.RGBA_INTEGER;
+			webgl_texture_src_type = gl.UNSIGNED_BYTE;
+			raw_channel_bit_depth = 8;
+			raw_channel_count = 2;
+			raw_channel_format = 'uint';
+		}
 		else if (ue_pixel_format === 'PF_R8_UINT')
 		{
 			webgl_texture_internal_format = gl.R8UI;
