Adds an error in the GPU dump viewer when the resource type is not supported.

diff --git a/Engine/Extras/GPUDumpViewer/GPUDumpViewer.html b/Engine/Extras/GPUDumpViewer/GPUDumpViewer.html
--- a/Engine/Extras/GPUDumpViewer/GPUDumpViewer.html
+++ b/Engine/Extras/GPUDumpViewer/GPUDumpViewer.html
@@ -1203,6 +1203,10 @@ function display_resource_internal(subresource_version_info)
 
 		g_view.set_resource_view(new_texture_view);
 	}
+	else
+	{
+		add_console_event('error', `Unable to display ${resource_desc['Name']}: ${resource_desc['Desc']} of type ${resource_desc['Type']} is not supported`);
+	}
 }
 
 
