Fixes PF_R8G8_UINT visualization

diff --git a/Engine/Extras/GPUDumpViewer/GPUDumpViewer.html b/Engine/Extras/GPUDumpViewer/GPUDumpViewer.html
--- a/Engine/Extras/GPUDumpViewer/GPUDumpViewer.html
+++ b/Engine/Extras/GPUDumpViewer/GPUDumpViewer.html
@@ -3137,7 +3137,7 @@ void main(void)
 		else if (ue_pixel_format === 'PF_R8G8_UINT')
 		{
 			webgl_texture_internal_format = gl.RG8UI;
-			webgl_texture_src_format = gl.RGBA_INTEGER;
+			webgl_texture_src_format = gl.RG_INTEGER;
 			webgl_texture_src_type = gl.UNSIGNED_BYTE;
 			raw_channel_bit_depth = 8;
 			raw_channel_count = 2;
@@ -3226,7 +3226,7 @@ void main(void)
 		{
 			webgl_channel_count = 1;
 		}
-		else if (webgl_texture_src_format === gl.RG)
+		else if (webgl_texture_src_format === gl.RG || webgl_texture_src_format === gl.RG_INTEGER)
 		{
 			webgl_channel_count = 2;
 		}
