Propagates texture zoom settings between passes and tabs

Submitting on behalf of Hugo Lamarche

diff --git a/Engine/Extras/GPUDumpViewer/GPUDumpViewer.html b/Engine/Extras/GPUDumpViewer/GPUDumpViewer.html
--- a/Engine/Extras/GPUDumpViewer/GPUDumpViewer.html
+++ b/Engine/Extras/GPUDumpViewer/GPUDumpViewer.html
@@ -15,6 +15,8 @@ const k_style_main_div_margin = 15;
 
 const k_null_json_ptr = "00000000000000000000";
 
+const k_texture_zoom_local_storage = 'texture_zoom';
+
 
 // -------------------------------------------------------------------- GLOBALS
 
@@ -247,6 +249,75 @@ function load_structure_metadata(structure_ptr)
 }
 
 
+// -------------------------------------------------------------------- LOCAL STORAGE
+
+function set_local_storage(name, value)
+{
+	if (typeof value === 'object')
+	{
+		value = JSON.stringify(value);
+	}
+	else if (typeof value === 'string')
+	{
+		// NOP
+	}
+	else
+	{
+		console.error('unknown type');
+	}
+
+	localStorage.setItem(name, value);
+}
+
+function get_local_storage(name)
+{
+	var value = localStorage.getItem(name);
+	if (value == null)
+	{
+		return null;
+	}
+
+	if (value[0] === '{')
+	{
+		value = JSON.parse(value);
+	}
+	return value;
+}
+
+function remove_local_storage(name)
+{
+	return localStorage.removeItem(name);
+}
+
+function clear_local_storage()
+{
+	localStorage.clear();
+}
+
+function test_local_storage()
+{
+	clear_local_storage();
+	console.assert(get_local_storage("doesntexist") == null);
+
+	set_local_storage("hello", "world");
+	console.assert(get_local_storage("hello") == "world");
+
+	remove_local_storage("hello");
+	console.assert(get_local_storage("hello") == null);
+
+	set_local_storage("hello", "world");
+	console.assert(get_local_storage("hello") == "world");
+
+	clear_local_storage();
+	console.assert(get_local_storage("hello") == null);
+
+	set_local_storage("hello", {"world": "cat"});
+	console.assert(get_local_storage("hello")["world"] == "cat");
+
+	clear_local_storage();
+}
+
+
 // -------------------------------------------------------------------- UTILITY
 
 function get_filename(file_path)
@@ -1189,11 +1260,28 @@ function display_resource_internal(subresource_version_info)
 		var new_texture_view = new TextureView(subresource_version_info, resource_desc);
 
 		// Keep the same zoom settings if the subresources are exactly the same extent
-		if (g_view && 'resource_view' in g_view && g_view.resource_view instanceof TextureView)
 		{
-			var zoom_settings = g_view.resource_view.get_zoom_settings();
+			var zoom_settings = null;
 
-			if (new_texture_view.is_zoom_settings_compatible(zoom_settings))
+			if (!zoom_settings && g_view && 'resource_view' in g_view && g_view.resource_view instanceof TextureView)
+			{
+				var prev_zoom_settings = g_view.resource_view.get_zoom_settings();
+				if (new_texture_view.is_zoom_settings_compatible(prev_zoom_settings))
+				{
+					zoom_settings = prev_zoom_settings;
+				}
+			}
+
+			if (!zoom_settings)
+			{
+				var global_zoom_settings = get_local_storage(k_texture_zoom_local_storage);
+				if (global_zoom_settings && new_texture_view.is_zoom_settings_compatible(global_zoom_settings))
+				{
+					zoom_settings = global_zoom_settings;
+				}
+			}
+
+			if (zoom_settings)
 			{
 				new_texture_view.onload = function() {
 					new_texture_view.apply_zoom_settings(zoom_settings);
@@ -2128,6 +2216,7 @@ class TextureView extends ResourceView
 			this.canvas.oncontextmenu = function(event) { return false; };
 
 			document.getElementById('canvas_viewport').oncontextmenu = function(event) { return false; };
+			document.getElementById('canvas_viewport').onscrollend = function(event) { texture_view.save_zoom_settings(); };
 
 			document.getElementById('canvas_outter').onclick = function(event) { texture_view.onclick_canvas_outter(); };
 			document.getElementById('canvas_outter').onmouseleave = function(event) { texture_view.set_select_viewport(false); };
@@ -2465,6 +2554,7 @@ class TextureView extends ResourceView
 
 		this.resize_canvas(new_pixel_scaling, texel_x, texel_y);
 		this.update_change_pixel_scaling_button();
+		this.save_zoom_settings();
 
 		return false;
 	}
@@ -2558,6 +2648,7 @@ class TextureView extends ResourceView
 		var canvas_viewport = document.getElementById('canvas_viewport');
 		canvas_viewport.scrollLeft -= event.movementX;
 		canvas_viewport.scrollTop -= event.movementY;
+		this.save_zoom_settings();
 	}
 
 	set_select_viewport(selected)
@@ -2663,6 +2754,8 @@ class TextureView extends ResourceView
 		canvas_viewport.scrollLeft = canvas_scroll_x;
 		canvas_viewport.scrollTop = canvas_scroll_y;
 		this.pixel_scaling = new_pixel_scaling;
+
+		this.save_zoom_settings();
 	}
 
 	get_zoom_settings()
@@ -2701,6 +2794,12 @@ class TextureView extends ResourceView
 		canvas_viewport.scrollTop = zoom_settings['canvas_viewport_scoll_top'];
 
 		this.update_change_pixel_scaling_button();
+		this.save_zoom_settings();
+	}
+
+	save_zoom_settings()
+	{
+		set_local_storage(k_texture_zoom_local_storage, this.get_zoom_settings());
 	}
 
 	init_gl()
@@ -5486,6 +5585,28 @@ function render_selection_list_html(parent_dom, display_list, options)
 }
 
 
+// -------------------------------------------------------------------- SWITCH TAB IN BROWSER
+
+document.addEventListener("visibilitychange", function() {
+	if (document.hidden)
+	{
+		// NOP
+	}
+	else
+	{
+		// apply the last zoom setting from last used tab when switting back to this tab.
+		if (g_view && 'resource_view' in g_view && g_view.resource_view instanceof TextureView)
+		{
+			var zoom_settings = get_local_storage(k_texture_zoom_local_storage);
+			if (zoom_settings && g_view.resource_view.is_zoom_settings_compatible(zoom_settings))
+			{
+				g_view.resource_view.apply_zoom_settings(zoom_settings);
+			}
+		}
+	}
+});
+
+
 // -------------------------------------------------------------------- WINDOW LAYOUT
 
 function init_top()
