Changes DumpGPU textures to be top left cornered

#tests win64

diff --git a/Engine/Extras/GPUDumpViewer/GPUDumpViewer.html b/Engine/Extras/GPUDumpViewer/GPUDumpViewer.html
--- a/Engine/Extras/GPUDumpViewer/GPUDumpViewer.html
+++ b/Engine/Extras/GPUDumpViewer/GPUDumpViewer.html
@@ -2289,7 +2289,7 @@ class TextureView extends ResourceView
 		console.assert(this.is_ready);
 
 		var extent = this.get_subresource_extent();
-		var pixel_data_offset = texel_x + (extent['y'] - 1 - texel_y) * extent['x'];
+		var pixel_data_offset = texel_x + texel_y * extent['x'];
 
 		var texel_value = [0];
 		for (var c = 1; c < this.subresource_desc['raw_channel_count']; c++)
@@ -3677,7 +3677,8 @@ in vec3 coordinates;
 out highp vec2 uv;
 void main(void) {
 	gl_Position = vec4(coordinates, 1.0);
-	uv = coordinates.xy * 0.5 + 0.5;
+	uv.x = 0.5 + coordinates.x * 0.5;
+	uv.y = 0.5 - coordinates.y * 0.5;
 }`;
 
 		var shader_program = gl_create_shader_program(this.gl, vert_code, frag_code);
diff --git a/Engine/Source/Runtime/RenderCore/Private/DumpGPU.cpp b/Engine/Source/Runtime/RenderCore/Private/DumpGPU.cpp
--- a/Engine/Source/Runtime/RenderCore/Private/DumpGPU.cpp
+++ b/Engine/Source/Runtime/RenderCore/Private/DumpGPU.cpp
@@ -948,8 +948,8 @@ public:
 
 		for (int32 y = 0; y < SubresourceDumpDesc.SubResourceExtent.Y; y++)
 		{
-			// Flip the data to be bottom left corner for the WebGL viewer.
-			const uint8* SrcPos = SrcData + SIZE_T(SubresourceDumpDesc.SubResourceExtent.Y - 1 - y) * SIZE_T(RowPitchInPixels) * BytePerPixel;
+			// Keep the data to be top left corner.
+			const uint8* SrcPos = SrcData + SIZE_T(y) * SIZE_T(RowPitchInPixels) * BytePerPixel;
 			uint8* DstPos = (&Array[0]) + SIZE_T(y) * SIZE_T(SubresourceDumpDesc.SubResourceExtent.X) * BytePerPixel;
 
 			FPlatformMemory::Memmove(DstPos, SrcPos, SIZE_T(SubresourceDumpDesc.SubResourceExtent.X) * BytePerPixel);
