GPUDump: Visualize R64 as R32G32 for now instead of nothing at all (it often contains two side by side 32bit values anyhow).

diff --git a/Engine/Extras/GPUDumpViewer/GPUDumpViewer.html b/Engine/Extras/GPUDumpViewer/GPUDumpViewer.html
--- a/Engine/Extras/GPUDumpViewer/GPUDumpViewer.html
+++ b/Engine/Extras/GPUDumpViewer/GPUDumpViewer.html
@@ -3182,7 +3182,7 @@ void main(void)
 			raw_channel_count = 4;
 			raw_channel_format = 'uint';
 		}
-		else if (ue_pixel_format === 'PF_R32G32_UINT')
+		else if (ue_pixel_format === 'PF_R32G32_UINT' || ue_pixel_format === 'PF_R64_UINT')
 		{
 			webgl_texture_internal_format = gl.RGBA8UI;
 			webgl_texture_src_format = gl.RGBA_INTEGER;
