Improves DumpGPU's sRGB and Rec709 display

diff --git a/Engine/Extras/GPUDumpViewer/GPUDumpViewer.html b/Engine/Extras/GPUDumpViewer/GPUDumpViewer.html
--- a/Engine/Extras/GPUDumpViewer/GPUDumpViewer.html
+++ b/Engine/Extras/GPUDumpViewer/GPUDumpViewer.html
@@ -1812,6 +1812,16 @@ function display_log()
 
 // -------------------------------------------------------------------- WEBGL UTILS
 
+function verify_backbuffer_color_space(gl)
+{
+	// https://developer.mozilla.org/en-US/docs/Web/API/WebGLRenderingContext/drawingBufferColorSpace
+	const desired_color_space = 'srgb';
+	if (gl.drawingBufferColorSpace !== desired_color_space)
+	{
+		add_console_event('error', `WebGL backbuffer's color space isn't ${desired_color_space}, it is ${gl.drawingBufferColorSpace}`);
+	}
+}
+
 function copy_canvas_to_clipboard(canvas, text)
 {
 	canvas.toBlob(function(image_blob) {
@@ -1884,6 +1894,41 @@ FloatInfos decodeFloat(float x)
 #define isinf(x) decodeFloat(x).is_infinity
 #define isdenorm(x) decodeFloat(x).is_denormal
 
+vec3 LinearToGamma22(vec3 Color)
+{
+	return pow(Color, vec3(2.2));
+}
+
+vec3 LinearToRec709(vec3 lin) 
+{
+	lin = max(vec3(6.10352e-5), lin);
+	return min(lin * vec3(4.5), pow(max(lin, vec3(0.018)), vec3(0.45)) * vec3(1.099) - vec3(0.099));
+}
+
+vec3 Rec709ToLinear(vec3 Color)
+{
+	Color = max(vec3(6.10352e-5), Color);
+	Color.r = Color.r > 0.081 ? pow((Color.r + 0.099) / 1.099, (1.0 / 0.45)) : Color.r / 4.5;
+	Color.g = Color.g > 0.081 ? pow((Color.g + 0.099) / 1.099, (1.0 / 0.45)) : Color.g / 4.5;
+	Color.b = Color.b > 0.081 ? pow((Color.b + 0.099) / 1.099, (1.0 / 0.45)) : Color.b / 4.5;
+	return Color;
+}
+
+vec3 LinearToSRGB(vec3 lin) 
+{
+	lin = max(vec3(6.10352e-5), lin);
+	return min(lin * vec3(12.92), pow(max(lin, vec3(0.00313067)), vec3(1.0/2.4)) * vec3(1.055) - vec3(0.055));
+}
+
+vec3 SRGBToLinear(vec3 Color) 
+{
+	Color = max(vec3(6.10352e-5), Color);
+	Color.r = Color.r > 0.04045 ? pow(Color.r * (1.0 / 1.055) + 0.0521327, 2.4) : Color.r * (1.0 / 12.92);
+	Color.g = Color.g > 0.04045 ? pow(Color.g * (1.0 / 1.055) + 0.0521327, 2.4) : Color.g * (1.0 / 12.92);
+	Color.b = Color.b > 0.04045 ? pow(Color.b * (1.0 / 1.055) + 0.0521327, 2.4) : Color.b * (1.0 / 12.92);
+	return Color;
+}
+
 `;
 	
 	var gl_shader = gl.createShader(shader_type);
@@ -2076,7 +2121,7 @@ class TextureView extends ResourceView
 		this.viewport_selected = false;
 		this.is_draging_canvas = false;
 		this.src_channels = 'RGB';
-		this.src_color_space = 'sRGB/Rec709';
+		this.src_color_space = 'Gamut=sRGB Gamma 2.2';
 		this.display_mode = 'Visualization';
 	}
 
@@ -2140,8 +2185,10 @@ class TextureView extends ResourceView
 						</div>
 						<div class="button_list" id="texture_visualization_change_src_color"><!---
 							---><span>Src color:</span><!---
-							---><input type="button" onclick="g_view.resource_view.change_src_color(this);" value="sRGB/Rec709" title="Performs no modifications to display to WebGL's sRGB backbuffer (default)."/><!---
-							---><input type="button" onclick="g_view.resource_view.change_src_color(this);" value="RGB Linear[0;1]" title="Applies a pow(display.rgb, 2.2); to output WebGL's sRGB backbuffer."/>
+							---><input type="button" onclick="g_view.resource_view.change_src_color(this);" value="Gamut=sRGB Gamma 2.2" title="Performs no modifications to display to WebGL's sRGB gamut backbuffer (default)."/><!---
+							---><input type="button" onclick="g_view.resource_view.change_src_color(this);" value="Gamut=sRGB Linear[0;1]" title="Applies a LinearToGamma22(display.rgb); to output WebGL's sRGB gamut backbuffer."/><!---
+							---><input type="button" onclick="g_view.resource_view.change_src_color(this);" value="sRGB" title="Applies a LinearToGamma22(SRGBToLinear(display.rgba)) To output to WebGL's sRGB gamut backbuffer."/><!---
+							---><input type="button" onclick="g_view.resource_view.change_src_color(this);" value="Rec709" title="Applies LinearToGamma22(Rec709ToLinear(display.rgb)) to output WebGL's sRGB gamut backbuffer."/>
 						</div>
 						<div class="button_list" id="texture_visualization_change_display_modes"><!---
 							---><input type="button" onclick="g_view.resource_view.change_display_mode(this);" value="Visualization" title="Display the WebGL's visualization shader's display.rgb."/><!---
@@ -2826,6 +2873,8 @@ class TextureView extends ResourceView
 		var gl = this.canvas.getContext('webgl2', gl_ctx_attributes);
 		this.gl = gl;
 
+		verify_backbuffer_color_space(this.gl);
+
 		// Init WebGL extensions
 		{
 			var available_extensions = this.gl.getSupportedExtensions();
@@ -3733,14 +3782,22 @@ void main(void) {
 				final_frag_code += 'display.rgb = display.aaa;';
 			}
 
-			if (this.src_color_space == 'sRGB/Rec709')
+			// verify_backbuffer_color_space() check webgl is a sRGB gamut backbuffer
+			if (this.src_color_space == 'Gamut=sRGB Gamma 2.2')
+			{
+				// NOP because verify_backbuffer_color_space() check webgl is a sRGB backbuffer
+			}
+			else if (this.src_color_space == 'Gamut=sRGB Linear[0;1]')
 			{
-				// NOP because src gamma = webgl display gamma.
+				final_frag_code += 'display.rgb = LinearToGamma22(display.rgb);';
 			}
-			else if (this.src_color_space == 'RGB Linear[0;1]')
+			else if (this.src_color_space == 'sRGB')
 			{
-				// Apply webgl's 2.2 gamma curve on the display color so linear gradients look linear in web browser.
-				final_frag_code += 'display.rgb = pow(display.rgb, vec3(2.2));';
+				final_frag_code += 'display.rgb = LinearToGamma22(SRGBToLinear(display.rgb));';
+			}
+			else if (this.src_color_space == 'Rec709')
+			{
+				final_frag_code += 'display.rgb = LinearToGamma22(Rec709ToLinear(display.rgb));';
 			}
 			final_frag_code += '}';
 
@@ -4556,6 +4613,8 @@ class Generic3DStructureView extends ResourceView
 		var gl = this.canvas.getContext('webgl2', gl_ctx_attributes);
 		this.gl = gl;
 
+		verify_backbuffer_color_space(gl);
+
 		// Init WebGL extensions
 		{
 			var available_extensions = this.gl.getSupportedExtensions();
