Fixed sRGB transfer functions (gamma) clipping values below 0.0000610352 (half of PR #12444).

Issue with non-zero HDR composited UI colors was confirmed, and this CL was validated against DX11/SM5 rendering engine tests.

diff --git a/Engine/Extras/GPUDumpViewer/GPUDumpViewer.html b/Engine/Extras/GPUDumpViewer/GPUDumpViewer.html
--- a/Engine/Extras/GPUDumpViewer/GPUDumpViewer.html
+++ b/Engine/Extras/GPUDumpViewer/GPUDumpViewer.html
@@ -1902,13 +1902,11 @@ vec3 LinearToWebGLDisplayGamma(vec3 Color)
 
 vec3 LinearToRec709(vec3 lin) 
 {
-	lin = max(vec3(6.10352e-5), lin);
 	return min(lin * vec3(4.5), pow(max(lin, vec3(0.018)), vec3(0.45)) * vec3(1.099) - vec3(0.099));
 }
 
 vec3 Rec709ToLinear(vec3 Color)
 {
-	Color = max(vec3(6.10352e-5), Color);
 	Color.r = Color.r > 0.081 ? pow((Color.r + 0.099) / 1.099, (1.0 / 0.45)) : Color.r / 4.5;
 	Color.g = Color.g > 0.081 ? pow((Color.g + 0.099) / 1.099, (1.0 / 0.45)) : Color.g / 4.5;
 	Color.b = Color.b > 0.081 ? pow((Color.b + 0.099) / 1.099, (1.0 / 0.45)) : Color.b / 4.5;
@@ -1917,13 +1915,11 @@ vec3 Rec709ToLinear(vec3 Color)
 
 vec3 LinearToSRGB(vec3 lin) 
 {
-	lin = max(vec3(6.10352e-5), lin);
 	return min(lin * vec3(12.92), pow(max(lin, vec3(0.00313067)), vec3(1.0/2.4)) * vec3(1.055) - vec3(0.055));
 }
 
 vec3 SRGBToLinear(vec3 Color) 
 {
-	Color = max(vec3(6.10352e-5), Color);
 	Color.r = Color.r > 0.04045 ? pow(Color.r * (1.0 / 1.055) + 0.0521327, 2.4) : Color.r * (1.0 / 12.92);
 	Color.g = Color.g > 0.04045 ? pow(Color.g * (1.0 / 1.055) + 0.0521327, 2.4) : Color.g * (1.0 / 12.92);
 	Color.b = Color.b > 0.04045 ? pow(Color.b * (1.0 / 1.055) + 0.0521327, 2.4) : Color.b * (1.0 / 12.92);
diff --git a/Engine/Shaders/Private/GammaCorrectionCommon.ush b/Engine/Shaders/Private/GammaCorrectionCommon.ush
--- a/Engine/Shaders/Private/GammaCorrectionCommon.ush
+++ b/Engine/Shaders/Private/GammaCorrectionCommon.ush
@@ -4,19 +4,16 @@
 
 half3 LinearTo709Branchless(half3 lin) 
 {
-	lin = max(6.10352e-5, lin); // minimum positive non-denormal (fixes black problem on DX11 AMD and NV)
 	return min(lin * 4.5, pow(max(lin, 0.018), 0.45) * 1.099 - 0.099);
 }
 
 half3 Rec709ToLinear(half3 Color)
 {
-	Color = max(6.10352e-5, Color); // minimum positive non-denormal (fixes black problem on DX11 AMD and NV)
 	return select(Color > 0.081, pow((Color + 0.099) / 1.099, 1.0 / 0.45), Color / 4.5);
 }
 
 half3 LinearToSrgbBranchless(half3 lin) 
 {
-	lin = max(6.10352e-5, lin); // minimum positive non-denormal (fixes black problem on DX11 AMD and NV)
 	return min(lin * 12.92, pow(max(lin, 0.00313067), 1.0/2.4) * 1.055 - 0.055);
 	// Possible that mobile GPUs might have native pow() function?
 	//return min(lin * 12.92, exp2(log2(max(lin, 0.00313067)) * (1.0/2.4) + log2(1.055)) - 0.055);
@@ -49,7 +46,6 @@ half3 LinearToSrgb(half3 lin)
 
 half3 sRGBToLinear( half3 Color ) 
 {
-	Color = max(6.10352e-5, Color); // minimum positive non-denormal (fixes black problem on DX11 AMD and NV)
 	return select(Color > 0.04045, pow( Color * (1.0 / 1.055) + 0.0521327, 2.4 ), Color * (1.0 / 12.92));
 }
 
diff --git a/Engine/Shaders/StandaloneRenderer/D3D/GammaCorrectionCommon.hlsl b/Engine/Shaders/StandaloneRenderer/D3D/GammaCorrectionCommon.hlsl
--- a/Engine/Shaders/StandaloneRenderer/D3D/GammaCorrectionCommon.hlsl
+++ b/Engine/Shaders/StandaloneRenderer/D3D/GammaCorrectionCommon.hlsl
@@ -3,13 +3,11 @@
 
 half3 LinearTo709Branchless(half3 lin) 
 {
-	lin = max(6.10352e-5, lin); // minimum positive non-denormal (fixes black problem on DX11 AMD and NV)
 	return min(lin * 4.5, pow(max(lin, 0.018), 0.45) * 1.099 - 0.099);
 }
 
 half3 LinearToSrgbBranchless(half3 lin) 
 {
-	lin = max(6.10352e-5, lin); // minimum positive non-denormal (fixes black problem on DX11 AMD and NV)
 	return min(lin * 12.92, pow(max(lin, 0.00313067), 1.0/2.4) * 1.055 - 0.055);
 	// Possible that mobile GPUs might have native pow() function?
 	//return min(lin * 12.92, exp2(log2(max(lin, 0.00313067)) * (1.0/2.4) + log2(1.055)) - 0.055);
@@ -31,7 +29,6 @@ half3 LinearToSrgbBranching(half3 lin)
 
 half3 sRGBToLinear( half3 Color ) 
 {
-	Color = max(6.10352e-5, Color); // minimum positive non-denormal (fixes black problem on DX11 AMD and NV)
 	return Color > 0.04045 ? pow( Color * (1.0 / 1.055) + 0.0521327, 2.4 ) : Color * (1.0 / 12.92);
 }
 
diff --git a/Engine/Shaders/StandaloneRenderer/OpenGL/SlateElementPixelShader.glsl b/Engine/Shaders/StandaloneRenderer/OpenGL/SlateElementPixelShader.glsl
--- a/Engine/Shaders/StandaloneRenderer/OpenGL/SlateElementPixelShader.glsl
+++ b/Engine/Shaders/StandaloneRenderer/OpenGL/SlateElementPixelShader.glsl
@@ -59,13 +59,11 @@ vec3 powScalar(vec3 values, float power)
 
 vec3 LinearTo709Branchless(vec3 lin)
 {
-	lin = maxWithScalar(6.10352e-5, lin); // minimum positive non-denormal (fixes black problem on DX11 AMD and NV)
 	return min(lin * 4.5, powScalar(maxWithScalar(0.018, lin), 0.45) * 1.099 - 0.099);
 }
 
 vec3 LinearToSrgbBranchless(vec3 lin)
 {
-	lin = maxWithScalar(6.10352e-5, lin); // minimum positive non-denormal (fixes black problem on DX11 AMD and NV)
 	return min(lin * 12.92, powScalar(maxWithScalar(0.00313067, lin), 1.0/2.4) * 1.055 - 0.055);
 	// Possible that mobile GPUs might have native pow() function?
 	//return min(lin * 12.92, exp2(log2(max(lin, 0.00313067)) * (1.0/2.4) + log2(1.055)) - 0.055);
