[DumpGPU] Add relative zoom toggle for cross-resolution texture comparison.

Added a "Relative Zoom" checkbox to the texture visualization bar. When enabled. switching between textures of different dimensions preserves the same canvas size and scroll position.
This keeps the view region spatially consistent across resolutions.

#Synced-CL 51948989

diff --git a/Engine/Extras/GPUDumpViewer/GPUDumpViewer.html b/Engine/Extras/GPUDumpViewer/GPUDumpViewer.html
--- a/Engine/Extras/GPUDumpViewer/GPUDumpViewer.html
+++ b/Engine/Extras/GPUDumpViewer/GPUDumpViewer.html
@@ -16,7 +16,7 @@ const k_style_main_div_margin = 15;
 const k_null_json_ptr = "00000000000000000000";
 
 const k_texture_zoom_local_storage = 'texture_zoom';
-
+const k_relative_zoom_local_storage = 'relative_zoom';
 
 // -------------------------------------------------------------------- GLOBALS
 
@@ -2207,6 +2207,13 @@ class TextureView extends ResourceView
 						<div class="button_list">
 							<input type="button" onclick="g_view.resource_view.copy_canvas_to_clipboard();" value="Copy to clipboard"/>
 						</div>
+						<div class="button_list">
+							<label title="When enabled, switching between textures of different dimensions preserves the relative zoom position (useful for comparing pre/post upscaling).">
+                                <input type="checkbox" id="relative_zoom_checkbox"
+                                    onchange="set_local_storage(k_relative_zoom_local_storage, this.checked ? '1' : '0');"
+                                /> Relative zoom
+                            </label>
+						</div>
 					</div>
 					<div id="canvas_viewport">
 						<canvas
@@ -2278,6 +2285,12 @@ class TextureView extends ResourceView
 			document.getElementById('canvas_outter').onmouseleave = function(event) { texture_view.set_select_viewport(false); };
 		}
 
+		// Restore relative zoom checkbox state
+		var relative_zoom_checkbox = document.getElementById('relative_zoom_checkbox');
+		if (relative_zoom_checkbox) {
+			relative_zoom_checkbox.checked = (get_local_storage(k_relative_zoom_local_storage) === '1');
+		}
+
 		if (this.subresource_desc['raw_channel_format'] == 'float')
 		{
 			update_value_selection(document.getElementById('texture_visualization_change_display_modes'), this.display_mode);
@@ -2579,7 +2592,7 @@ class TextureView extends ResourceView
 
 			if (zooming_in)
 			{
-				if (this.pixel_scaling == 0)
+				if (this.pixel_scaling <= 0)
 				{
 					if (min_pixel_scaling <= 16)
 					{
@@ -2625,7 +2638,8 @@ class TextureView extends ResourceView
 			8: '8:1',
 			16: '16:1',
 		};
-		update_value_selection(document.getElementById('texture_visualization_change_pixel_scaling'), pixel_scaling_names[this.pixel_scaling]);
+		var name = pixel_scaling_names[this.pixel_scaling] || '';
+		update_value_selection(document.getElementById('texture_visualization_change_pixel_scaling'), name);
 	}
 
 	canvas_click(event)
@@ -2834,6 +2848,11 @@ class TextureView extends ResourceView
 
 	is_zoom_settings_compatible(zoom_settings)
 	{
+		// Always compatible when relative zoom is enabled
+		if (get_local_storage(k_relative_zoom_local_storage) === '1') {
+			return true;
+		}
+
 		var subresource_extent = this.get_subresource_extent();
 
 		return (zoom_settings['subresource_extent_x'] == subresource_extent['x'] && zoom_settings['subresource_extent_y'] == subresource_extent['y']);
@@ -2843,11 +2862,31 @@ class TextureView extends ResourceView
 	{
 		var canvas_viewport = document.getElementById('canvas_viewport');
 
-		this.pixel_scaling = zoom_settings['pixel_scaling'];
-		this.canvas.style.width = zoom_settings['canvas_width'];
-		this.canvas.style.height = zoom_settings['canvas_height'];
-		canvas_viewport.scrollLeft = zoom_settings['canvas_viewport_scoll_left'];
-		canvas_viewport.scrollTop = zoom_settings['canvas_viewport_scoll_top'];
+		var subresource_extent = this.get_subresource_extent();
+		var prev_extent_x = zoom_settings['subresource_extent_x'];
+		var prev_extent_y = zoom_settings['subresource_extent_y'];
+		var same_extent = (prev_extent_x == subresource_extent['x'] &&
+						   prev_extent_y == subresource_extent['y']);
+
+		if (same_extent || get_local_storage(k_relative_zoom_local_storage) !== '1')
+		{
+			this.pixel_scaling = zoom_settings['pixel_scaling'];
+			this.canvas.style.width = zoom_settings['canvas_width'];
+			this.canvas.style.height = zoom_settings['canvas_height'];
+			canvas_viewport.scrollLeft = zoom_settings['canvas_viewport_scoll_left'];
+			canvas_viewport.scrollTop = zoom_settings['canvas_viewport_scoll_top'];
+		}
+		else {
+			var prev_canvas_w = parseFloat(zoom_settings['canvas_width']);
+			var prev_canvas_h = parseFloat(zoom_settings['canvas_height']);
+
+			this.canvas.style.width = prev_canvas_w + 'px';
+			this.canvas.style.height = prev_canvas_h + 'px';
+			canvas_viewport.scrollLeft = zoom_settings['canvas_viewport_scoll_left'];
+			canvas_viewport.scrollTop = zoom_settings['canvas_viewport_scoll_top'];
+
+			this.pixel_scaling = -1;
+		}
 
 		this.update_change_pixel_scaling_button();
 		this.save_zoom_settings();
