Use explicit navigation commands instead of running arbitrary JS via eval().

Also fix a bad variable name when building the error string when a buffer cannot be loaded.

#rnx
#tests Clicked on all UI item types (passes, resources, parameters, resource history etc.) and verified that they display correctly.
#Synced-CL 52126719

diff --git a/Engine/Extras/GPUDumpViewer/GPUDumpViewer.html b/Engine/Extras/GPUDumpViewer/GPUDumpViewer.html
--- a/Engine/Extras/GPUDumpViewer/GPUDumpViewer.html
+++ b/Engine/Extras/GPUDumpViewer/GPUDumpViewer.html
@@ -787,7 +787,7 @@ function display_pass_hierarchy()
 				parent_event_scopes.push(scope);
 			}
 
-			html += `<a style="padding-left: ${10 + 16 * pass_event_scopes.length}px;" href="#display_pass(${pass_id});">${pass_data['EventName']}</a>`;
+			html += `<a style="padding-left: ${10 + 16 * pass_event_scopes.length}px;" href="#display_pass;${pass_id}">${pass_data['EventName']}</a>`;
 		}
 	});
 	document.getElementById('pass_hierarchy').innerHTML = html;
@@ -841,7 +841,7 @@ class PassView extends IView
 			<div class="pass_title main_div">
 				<div id="pass_frame_list" style="display:none;"></div>
 				${this.pass_data['EventName']}
-				<div class="button_list" style="display:inline-block;"><a href="#display_pass_parameters(${this.pass_id});">PassParameters</a></div>
+				<div class="button_list" style="display:inline-block;"><a href="#display_pass_parameters;${this.pass_id}">PassParameters</a></div>
 			</div>
 			<table width="100%">
 				<tr>
@@ -943,8 +943,8 @@ class PassView extends IView
 		var next_frame_href = '';
 		if (this.resource_view.resource_desc['Desc'] == 'FRDGParameterStruct')
 		{
-			prev_frame_href = `#display_pass_parameters(${prev_frame_pass_id});`;
-			next_frame_href = `#display_pass_parameters(${next_frame_pass_id});`;
+			prev_frame_href = `#display_pass_parameters;${prev_frame_pass_id}`;
+			next_frame_href = `#display_pass_parameters;${next_frame_pass_id}`;
 		}
 		else
 		{
@@ -964,13 +964,13 @@ class PassView extends IView
 			}
 			else if (is_output_resource)
 			{
-				prev_frame_href = `#display_output_resource(${prev_frame_pass_id},'${get_subresource_unique_name(prev_frame_subresource_info)}');`;
-				next_frame_href = `#display_output_resource(${next_frame_pass_id},'${get_subresource_unique_name(next_frame_subresource_info)}');`;
+				prev_frame_href = `#display_output_resource;${prev_frame_pass_id};${encodeURIComponent(get_subresource_unique_name(prev_frame_subresource_info))}`;
+				next_frame_href = `#display_output_resource;${next_frame_pass_id};${encodeURIComponent(get_subresource_unique_name(next_frame_subresource_info))}`;
 			}
 			else
 			{
-				prev_frame_href = `#display_input_resource(${prev_frame_pass_id},'${get_subresource_unique_name(prev_frame_subresource_info)}');`;
-				next_frame_href = `#display_input_resource(${next_frame_pass_id},'${get_subresource_unique_name(next_frame_subresource_info)}');`;
+				prev_frame_href = `#display_input_resource;${prev_frame_pass_id};${encodeURIComponent(get_subresource_unique_name(prev_frame_subresource_info))}`;
+				next_frame_href = `#display_input_resource;${next_frame_pass_id};${encodeURIComponent(get_subresource_unique_name(next_frame_subresource_info))}`;
 			}
 		}
 
@@ -992,7 +992,7 @@ class PassView extends IView
 			var href = null;
 			if (resource_desc)
 			{
-				href = `#display_input_resource(${this.pass_id},'${subresource_unique_name}');`
+				href = `#display_input_resource;${this.pass_id};${encodeURIComponent(subresource_unique_name)}`
 			}
 
 			display_list.push({'name': name, 'href': href});
@@ -1028,11 +1028,11 @@ class PassView extends IView
 			{
 				if (draw_id >= 0)
 				{
-					href = `#display_draw_output_resource(${this.pass_id},'${subresource_unique_name}',${draw_id});`;
+					href = `#display_draw_output_resource;${this.pass_id};${encodeURIComponent(subresource_unique_name)};${draw_id}`;
 				}
 				else
 				{
-					href = `#display_output_resource(${this.pass_id},'${subresource_unique_name}');`;
+					href = `#display_output_resource;${this.pass_id};${encodeURIComponent(subresource_unique_name)}`;
 				}
 			}
 
@@ -1063,7 +1063,7 @@ class PassView extends IView
 			var href = null;
 			if (subresource_unique_name && is_output_resource)
 			{
-				href = `#display_draw_output_resource(${this.pass_id},'${subresource_unique_name}',${draw_id});`;
+				href = `#display_draw_output_resource;${this.pass_id};${encodeURIComponent(subresource_unique_name)};${draw_id}`;
 			}
 
 			display_list.push({
@@ -1077,7 +1077,7 @@ class PassView extends IView
 			var href = null;
 			if (is_output_resource)
 			{
-				href = `#display_output_resource(${this.pass_id},'${subresource_unique_name}');`;
+				href = `#display_output_resource;${this.pass_id};${encodeURIComponent(subresource_unique_name)}`;
 			}
 
 			display_list.push({
@@ -1121,7 +1121,7 @@ class PassView extends IView
 			{
 				display_list.push({
 					'name': g_passes[pass_id]['EventName'],
-					'href': `#display_output_resource(${pass_id},'${subresource_unique_name}');`,
+					'href': `#display_output_resource;${pass_id};${encodeURIComponent(subresource_unique_name)}`,
 				});
 			}
 		}
@@ -1162,7 +1162,7 @@ class PassView extends IView
 			{
 				display_list.push({
 					'name': g_passes[pass_id]['EventName'],
-					'href': `#display_input_resource(${pass_id},'${subresource_unique_name}');`,
+					'href': `#display_input_resource;${pass_id};${encodeURIComponent(subresource_unique_name)}`,
 				});
 			}
 		}
@@ -1205,7 +1205,7 @@ class PassView extends IView
 
 	get navigations()
 	{
-		var navs = [`display_pass(${this.pass_id});`];
+		var navs = [`display_pass;${this.pass_id}`];
 
 		if (this.resource_view !== null)
 		{
@@ -1305,11 +1305,11 @@ function display_pass(pass_id)
 	// Load first resource
 	if (g_passes[pass_id]['OutputResources'][0])
 	{
-		redirect_to_hash(`display_output_resource(${pass_id},'${g_passes[pass_id]['OutputResources'][0]}');`);
+		redirect_to_hash('display_output_resource', pass_id, g_passes[pass_id]['OutputResources'][0]);
 	}
 	else if (g_passes[pass_id]['InputResources'][0])
 	{
-		redirect_to_hash(`display_input_resource(${pass_id},'${g_passes[pass_id]['InputResources'][0]}');`);
+		redirect_to_hash('display_input_resource', pass_id, g_passes[pass_id]['InputResources'][0]);
 	}
 	else
 	{
@@ -1420,7 +1420,7 @@ class ConsoleView extends IView
 
 	get navigations()
 	{
-		return [`display_console();`];
+		return [`display_console`];
 	}
 }
 
@@ -1515,7 +1515,7 @@ class InfosView extends IView
 
 	get navigations()
 	{
-		return [`display_infos();`];
+		return [`display_infos`];
 	}
 }
 
@@ -1544,7 +1544,7 @@ const k_tips = [
 	'Make sure to browse the dump informations.',
 	'Make sure to browse the console variables.',
 	'Make sure to browse the log file.',
-	'Share the part of the URL after the # (for instance #display_input_resource(96,\'TSR.AntiAliasing.Noise.000000006b9b2c00.mip0\');) to anyone else whom have this GPU dump to so they too can navigate to the exact same resource view.',
+	'Share the part of the URL after the # (for instance #display_input_resource;96;TSR.AntiAliasing.Noise.000000006b9b2c00.mip0) to anyone else whom have this GPU dump to so they too can navigate to the exact same resource view.',
 	'Set r.DumpGPU.Viewer.Visualize in your ConsoleVariables.ini so the dump viewer automatically open this RDG output resource at startup.',
 
 	// Buffer visualization
@@ -1573,14 +1573,14 @@ class TipView extends IView
 				<div>User tip:</div>
 				<div style="padding: 20px;">${k_tips[this.tip_id]}</div>
 				<div class="button_list" style="margin-left: auto; display: block; width: 100px;">
-						<a href="#display_tip(${(this.tip_id + 1) % k_tips.length});">Next</a>
+						<a href="#display_tip;${(this.tip_id + 1) % k_tips.length}">Next</a>
 				</div>
 			</div>`;
 	}
 
 	get navigations()
 	{
-		return [`display_tip(${this.tip_id});`];
+		return [`display_tip;${this.tip_id}`];
 	}
 }
 
@@ -1703,7 +1703,7 @@ class CVarsView extends IView
 
 	get navigations()
 	{
-		return [`display_cvars();`];
+		return [`display_cvars`];
 	}
 
 	release()
@@ -1787,7 +1787,7 @@ class LogView extends IView
 
 	get navigations()
 	{
-		return [`display_log();`];
+		return [`display_log`];
 	}
 
 	resize(ctx)
@@ -4054,7 +4054,7 @@ class GenericBufferView extends ResourceView
 			else
 			{
 				document.getElementById('buffer_content_pannel').innerHTML = `
-					Error: ${this.raw_buffer_data_path} couldn't be loaded.`;
+					Error: ${buffer_view.raw_buffer_data_path} couldn't be loaded.`;
 			}
 		});
 
@@ -5729,11 +5729,11 @@ function init_top()
 function init_top_buttons()
 {
 	document.getElementById('top_buttons').innerHTML = `
-		<a href="#display_infos();">Infos</a>
-		<a href="#display_cvars();">CVars</a>
-		<a href="#display_log();">Log</a>`;
+		<a href="#display_infos">Infos</a>
+		<a href="#display_cvars">CVars</a>
+		<a href="#display_log">Log</a>`;
 	document.getElementById('top_viewer_buttons').innerHTML = `
-		<a id="console_button" href="#display_console();">Console</a>`;
+		<a id="console_button" href="#display_console">Console</a>`;
 	update_console_button();
 }
 
@@ -5891,7 +5891,7 @@ function init_page()
 
 		if (pass_id_to_open != -1)
 		{
-			return redirect_to_hash(`display_output_resource(${pass_id_to_open},'${subresource_unique_name_to_open}');`);
+			return redirect_to_hash('display_output_resource', pass_id_to_open, subresource_unique_name_to_open);
 		}
 	}
 	
@@ -5953,8 +5953,15 @@ function onresize_body()
 
 var g_previous_location_hash = null;
 
-function redirect_to_hash(new_hash)
+function redirect_to_hash(action)
 {
+	var parts = [action];
+	for (var i = 1; i < arguments.length; i++)
+	{
+		parts.push(encodeURIComponent(arguments[i]));
+	}
+	var new_hash = parts.join(';');
+
 	if (new_hash === g_previous_location_hash)
 	{
 		return;
@@ -5963,6 +5970,32 @@ function redirect_to_hash(new_hash)
 	location.hash = new_hash;
 }
 
+function dispatch_navigation(hash_str)
+{
+	var parts = hash_str.split(';');
+	var action = parts[0];
+	var args = [];
+	for (var i = 1; i < parts.length; i++)
+	{
+		args.push(decodeURIComponent(parts[i]));
+	}
+
+	switch (action)
+	{
+		case 'display_pass':                  display_pass(args[0]); break;
+		case 'display_pass_parameters':       display_pass_parameters(args[0]); break;
+		case 'display_tip':                   display_tip(args[0]); break;
+		case 'display_infos':                 display_infos(); break;
+		case 'display_cvars':                 display_cvars(); break;
+		case 'display_log':                   display_log(); break;
+		case 'display_console':               display_console(); break;
+		case 'display_output_resource':       display_output_resource(args[0], args[1]); break;
+		case 'display_input_resource':        display_input_resource(args[0], args[1]); break;
+		case 'display_draw_output_resource':  display_draw_output_resource(args[0], args[1], args[2]); break;
+		default:                              add_console_event('error', `Unknown navigation action: "${action}"`); break;
+	}
+}
+
 function navigate_to_hash()
 {
 	if (location.hash === g_previous_location_hash)
@@ -5974,8 +6007,8 @@ function navigate_to_hash()
 
 	if (location.hash)
 	{
-		var js = location.hash.substring(1);
-		eval(js);
+		var hash_str = location.hash.substring(1);
+		dispatch_navigation(hash_str);
 	}
 	else
 	{
