[DumpGPU] The hierarchy view will scroll to keep the selected pass in view when '-' or '+' is used.

#tests + - of multiple frames instantly scroll the target pass in view in the hierarchy
#Synced-CL 51948989

diff --git a/Engine/Extras/GPUDumpViewer/GPUDumpViewer.html b/Engine/Extras/GPUDumpViewer/GPUDumpViewer.html
--- a/Engine/Extras/GPUDumpViewer/GPUDumpViewer.html
+++ b/Engine/Extras/GPUDumpViewer/GPUDumpViewer.html
@@ -5597,6 +5597,17 @@ function update_href_selection(parent)
 			a.classList.remove('match_location_hash');
 		}
 	};
+
+	// Scroll the pass hierarchy view to keep the selected pass visible if '-' or '+' is used
+	var pass_hierarchy = document.getElementById('pass_hierarchy');
+	if (pass_hierarchy)
+	{
+		var selected = pass_hierarchy.querySelector('a.match_location_hash');
+		if (selected)
+		{
+			selected.scrollIntoView({block:'nearest', behavior:'instant'});
+		}
+	}
 }
 
 function update_value_selection(parent, selected_value)
