[DumpGPU] Now DumpGPUViewer can inspect the referenced uniform buffer if r.RHI.UniformBufferContentMap.Enable is set to true.

A new ReferencedStructs.json is generated in Base directory so that the GPUDumpViewer can read it.

#tests test on both r.RHI.UniformBufferContentMap.Enable disabled and enabled on single and two frames.
#Synced-CL 52438879

diff --git a/Engine/Extras/GPUDumpViewer/GPUDumpViewer.html b/Engine/Extras/GPUDumpViewer/GPUDumpViewer.html
--- a/Engine/Extras/GPUDumpViewer/GPUDumpViewer.html
+++ b/Engine/Extras/GPUDumpViewer/GPUDumpViewer.html
@@ -24,6 +24,7 @@ var g_dump_service = {};
 var g_infos = {};
 var g_dump_cvars = {};
 var g_passes = [];
+var g_referenced_structs = {};
 var g_descs = {};
 
 var g_view = null;
@@ -4185,6 +4186,10 @@ class GenericBufferView extends ResourceView
 				{
 					column_format = 'float';
 				}
+				else if (it.base_type == 'UBMT_REFERENCED_STRUCT')
+				{
+					column_format = 'referenced_struct';
+				}
 				else
 				{
 					return;
@@ -4391,7 +4396,15 @@ class GenericBufferView extends ResourceView
 
 					if (this.display_elem_using_rows && i == 1 && this.structure_metadata)
 					{
-						display = `hex(${display})`;
+						if (display != 'referenced_struct') {
+							display = `hex(${display})`;
+						}
+						else
+						{
+							value = '-';
+							resource_content += `<td>${value}</td>`;
+							continue;
+						}
 					}
 
 					if (display == 'float')
@@ -4450,6 +4463,40 @@ class GenericBufferView extends ResourceView
 					{
 						value = '0b' + (new Uint8Array(this.raw_buffer_data, byte_offset, 1))[0].toString(2).padStart(8, 0);
 					}
+					else if (display == 'referenced_struct')
+					{
+						// Check if this pass has referenced struct data for this member
+						var pass_pointer = this.subresource_version_info['pass'];
+						var ref_structs = g_referenced_structs[pass_pointer];
+						if (ref_structs && ref_structs[member_names[member_id]])
+						{
+							// Find the pass_id from the pass pointer
+							var link_pass_id = -1;
+							for (var pid = 0; pid < g_passes.length; pid++)
+							{
+								if (g_passes[pid]['Pointer'] == pass_pointer)
+								{
+									link_pass_id = pid;
+									break;
+								}
+							}
+
+							if (link_pass_id >= 0)
+							{
+								var ref_meta = ref_structs[member_names[member_id]];
+								var ref_struct_metadata = load_structure_metadata(ref_meta['Metadata']);
+								value = `<a href="#display_referenced_struct;${link_pass_id};${encodeURIComponent(member_names[member_id])}" title="Click to view ${ref_struct_metadata['StructTypeName']} uniform buffer contents">${ref_struct_metadata['StructTypeName']}</a>`;
+							}
+							else
+							{
+								value = '(referenced struct)';
+							}
+						}
+						else
+						{
+								value = '(not available. Enable r.RHI.UniformBufferContentMap.Enable)';
+						}
+					}
 					else
 					{
 						value = `Unknown ${display}`;
@@ -4550,6 +4597,48 @@ function display_pass_parameters(pass_id)
 	g_view.set_resource_view(new ParameterStructureView(subresource_version_info, resource_desc, structure_metadata));
 }
 
+// -------------------------------------------------------------------- REFERENCED STRUCT DISPLAY
+
+function display_referenced_struct(pass_id, member_name) {
+
+	display_pass_internal(pass_id);
+
+	var pass_data = g_passes[pass_id];
+	var pass_pointer = pass_data['Pointer'];
+
+	var ref_structs = g_referenced_structs[pass_pointer];
+	if (!ref_structs || !ref_structs[member_name]) {
+		set_main_view(null);
+		return;
+	}
+
+	var ref_entry = ref_structs[member_name];
+
+	if (!does_file_exists(`Structures/${ref_entry['Content']}.bin`))
+	{
+		add_console_event('error', `Referenced struct "${member_name}" binary not found. Set r.RHI.UniformBufferContentMap.Enable=true and dump again.`);
+		return;
+	}
+
+	var structure_metadata = load_structure_metadata(ref_entry['Metadata']);
+
+	var subresource_version_info = {};
+	subresource_version_info['subresource'] = null;
+	subresource_version_info['resource'] = ref_entry['Content'];
+	subresource_version_info['pass'] = pass_data['Pointer'];
+	subresource_version_info['draw'] = -1;
+
+	var resource_desc = {};
+	resource_desc['Name'] = member_name;
+	resource_desc['ByteSize'] = structure_metadata['Size'];
+	resource_desc['Desc'] = 'FRDGParameterStruct';
+	resource_desc['BytesPerElement'] = resource_desc['ByteSize'];
+	resource_desc['NumElements'] = 1;
+	resource_desc['Metadata'] = ref_entry['Metadata'];
+	resource_desc['Usage'] = [];
+
+	g_view.set_resource_view(new ParameterStructureView(subresource_version_info, resource_desc, structure_metadata));
+}
 
 // -------------------------------------------------------------------- GENERIC 3D STRUCTURE VIEW
 
@@ -5833,6 +5922,21 @@ function init_page()
 		}
 	}
 
+	// Load referenced struct mapping.
+	g_referenced_structs = {};
+	if (does_file_exists('Base/ReferencedStructs.json'))
+	{
+		try {
+			var ref_structs_list = load_json_dict_sequence('Base/ReferencedStructs.json');
+			for (var entry of ref_structs_list) {
+				g_referenced_structs[entry['Pass']] = entry['ReferencedStructs'];
+			}
+		}
+		catch (e) {
+			add_console_event('log', 'ReferencedStructs.json found but could not be parsed.');
+		}	
+	}
+
 	// Load all the resource descriptors
 	{
 		g_descs = {};
@@ -5995,6 +6099,7 @@ function dispatch_navigation(hash_str)
 	{
 		case 'display_pass':                  display_pass(args[0]); break;
 		case 'display_pass_parameters':       display_pass_parameters(args[0]); break;
+		case 'display_referenced_struct':     display_referenced_struct(args[0], args[1]); break;
 		case 'display_tip':                   display_tip(args[0]); break;
 		case 'display_infos':                 display_infos(); break;
 		case 'display_cvars':                 display_cvars(); break;
@@ -6667,6 +6772,18 @@ input::placeholder
 	color: var(--button-txt-color-selected);
 }
 
+#buffer_outter #buffer_viewport #buffer_content_table tr td a[href]
+{
+	color: var(--button-txt-color-selected);
+	cursor: pointer;
+	display: inline;
+}
+
+#buffer_outter #buffer_viewport #buffer_content_table tr td a[href]:hover
+{
+	color: var(--button-txt-color-hover);
+}
+
 </style>
 
 
diff --git a/Engine/Source/Runtime/RenderCore/Private/DumpGPU.cpp b/Engine/Source/Runtime/RenderCore/Private/DumpGPU.cpp
--- a/Engine/Source/Runtime/RenderCore/Private/DumpGPU.cpp
+++ b/Engine/Source/Runtime/RenderCore/Private/DumpGPU.cpp
@@ -2504,7 +2504,7 @@ static const TCHAR* GetPassEventNameWithGPUMask(const FRDGPass* Pass, FString& O
 }
 
 #if !(UE_BUILD_SHIPPING || UE_BUILD_TEST)
-void DumpReferencedStructs(const FShaderParametersMetadata* Metadata, const uint8* Content, uint32 ParentOffset)
+void DumpReferencedStructs(const FShaderParametersMetadata* Metadata, const uint8* Content, uint32 ParentOffset,const FString& Prefix, TSharedPtr<FJsonObject> OutReferencedStructs)
 {
 	check(Metadata);
 	check(Content);
@@ -2520,10 +2520,12 @@ void DumpReferencedStructs(const FShaderParametersMetadata* Metadata, const uint
 			continue;
 		}
 
+		FString MemberName = Prefix + Member.GetName();
+
 		if (Member.GetBaseType() != EUniformBufferBaseType::UBMT_REFERENCED_STRUCT)
 		{
 			// Inlined structs can be skipped as they are contained in the pass content but since they can contain referenced structs we recursively call DumpReferencedStructs
-			DumpReferencedStructs(MemberMetadata, Content, ParentOffset + Member.GetOffset());
+			DumpReferencedStructs(MemberMetadata, Content, ParentOffset + Member.GetOffset(), MemberName + TEXT("."), OutReferencedStructs);
 			continue;
 		}
 
@@ -2540,16 +2542,28 @@ void DumpReferencedStructs(const FShaderParametersMetadata* Metadata, const uint
 			// Skip referenced structs without a valid buffer
 			continue;
 		}
-		
+
+		const uint8* ReferencedStructContent = (const uint8*)UE::RHI::UniformBufferContentMap::Get(UniformBuffer);
+		// Record the mapping so the GPU dump viewer can link to this referenced struct's content
+		if (ReferencedStructContent)
+		{
+			if (OutReferencedStructs)
+			{
+				TSharedPtr<FJsonObject> RefStructEntry = MakeShareable(new FJsonObject);
+				RefStructEntry->SetStringField(TEXT("Content"), ResourceDumpContext->PtrToString(UniformBuffer));
+				RefStructEntry->SetStringField(TEXT("Metadata"), ResourceDumpContext->PtrToString(MemberMetadata));
+				OutReferencedStructs->SetObjectField(MemberName, RefStructEntry);
+			}
+		}
+
 		if (!ResourceDumpContext->IsDumped(UniformBuffer))
 		{
 			ResourceDumpContext->SetDumped(UniformBuffer);
 			ResourceDumpContext->Dump(MemberMetadata);
 
-			const uint8* ReferencedStructContent = (const uint8*)UE::RHI::UniformBufferContentMap::Get(UniformBuffer);
 			if (ReferencedStructContent)
 			{
-				DumpReferencedStructs(MemberMetadata, ReferencedStructContent, 0);
+				DumpReferencedStructs(MemberMetadata, ReferencedStructContent, 0, MemberName + TEXT("."), OutReferencedStructs);
 
 				uint32 PassParametersByteSize = MemberMetadata->GetSize();
 				if (PassParametersByteSize > 0)
@@ -2824,6 +2838,9 @@ void FRDGBuilder::DumpResourcePassOutputs(const FRDGPass* Pass)
 	// Dump the pass' parameters
 	if (GDumpGPUPassParameters.GetValueOnRenderThread() != 0)
 	{
+#if !(UE_BUILD_SHIPPING || UE_BUILD_TEST)
+		TSharedPtr<FJsonObject> ReferencedStructsJson = MakeShareable(new FJsonObject);
+#endif
 		int32 PassParametersByteSize = 0;
 		{
 			const FShaderParametersMetadata* Metadata = Pass->GetParameters().GetMetadata();
@@ -2834,7 +2851,7 @@ void FRDGBuilder::DumpResourcePassOutputs(const FRDGPass* Pass)
 #if !(UE_BUILD_SHIPPING || UE_BUILD_TEST)
 				if (const uint8* Content = Pass->GetParameters().GetContents())
 				{
-					DumpReferencedStructs(Metadata, Content, 0);
+					DumpReferencedStructs(Metadata, Content, 0, TEXT(""), ReferencedStructsJson);
 				}
 #endif
 			}
@@ -2853,6 +2870,18 @@ void FRDGBuilder::DumpResourcePassOutputs(const FRDGPass* Pass)
 			ResourceDumpContext->DumpBinaryToFile(ArrayView, DumpFilePath, FRDGResourceDumpContext::ETimingBucket::ParametersFileWrite);
 			ResourceDumpContext->SetDumped(PassParametersContent);
 		}
+#if !(UE_BUILD_SHIPPING || UE_BUILD_TEST)
+		// Write referenced struct mapping for the viewer if any were found
+		if (ReferencedStructsJson->Values.Num() > 0)
+		{
+			TSharedPtr<FJsonObject> ReferencedStructsFile = MakeShareable(new FJsonObject);
+			ReferencedStructsFile->SetStringField(TEXT("Pass"), ResourceDumpContext->PtrToString(Pass));
+			ReferencedStructsFile->SetObjectField(TEXT("ReferencedStructs"), ReferencedStructsJson);
+			FString RefStructsPath = FString(FRDGResourceDumpContext::kBaseDir) / TEXT("ReferencedStructs.json");
+			ResourceDumpContext->DumpJsonToFile(ReferencedStructsFile, RefStructsPath, FILEWRITE_Append);
+		}
+#endif
+
 	}
 
 	ResourceDumpContext->PassesCount++;
